CIS: Fix for Mac non-unity build

diff --git a/Engine/Plugins/Experimental/Workspace/Source/WorkspaceEditor/Private/Workspace.h b/Engine/Plugins/Experimental/Workspace/Source/WorkspaceEditor/Private/Workspace.h
--- a/Engine/Plugins/Experimental/Workspace/Source/WorkspaceEditor/Private/Workspace.h
+++ b/Engine/Plugins/Experimental/Workspace/Source/WorkspaceEditor/Private/Workspace.h
@@ -8,6 +8,7 @@
 
 struct FEditedDocumentInfo;
 class UWorkspace;
+class UWorkspaceState;
 class UWorkspaceSchema;
 class UWorkspaceFactory;
 class UAssetDefinition_Workspace;
