Fix crash cooking EngineTest

Dont register against the UWorkspaceState CDO and dont attempt to unregister in the destructor, the delegate invocation list will be compacted anyways.

diff --git a/Engine/Plugins/Experimental/Workspace/Source/WorkspaceEditor/Private/WorkspaceState.cpp b/Engine/Plugins/Experimental/Workspace/Source/WorkspaceEditor/Private/WorkspaceState.cpp
--- a/Engine/Plugins/Experimental/Workspace/Source/WorkspaceEditor/Private/WorkspaceState.cpp
+++ b/Engine/Plugins/Experimental/Workspace/Source/WorkspaceEditor/Private/WorkspaceState.cpp
@@ -13,14 +13,11 @@
 
 UWorkspaceState::UWorkspaceState()
 {
-	IAssetRegistry& AssetRegistry = FModuleManager::Get().LoadModuleChecked<FAssetRegistryModule>(TEXT("AssetRegistry")).Get();
-	AssetRegistry.OnAssetRenamed().AddUObject(this, &UWorkspaceState::HandleAssetRenamed);
-}
-
-UWorkspaceState::~UWorkspaceState()
-{
-	IAssetRegistry& AssetRegistry = FModuleManager::Get().LoadModuleChecked<FAssetRegistryModule>(TEXT("AssetRegistry")).Get();
-	AssetRegistry.OnAssetRenamed().RemoveAll(this);
+	if(!HasAnyFlags(RF_ClassDefaultObject))
+	{
+		IAssetRegistry& AssetRegistry = FModuleManager::Get().LoadModuleChecked<FAssetRegistryModule>(TEXT("AssetRegistry")).Get();
+		AssetRegistry.OnAssetRenamed().AddUObject(this, &UWorkspaceState::HandleAssetRenamed);
+	}
 }
 
 FString UWorkspaceState::GetWorkspacePath(const UWorkspace* InWorkspace)
diff --git a/Engine/Plugins/Experimental/Workspace/Source/WorkspaceEditor/Private/WorkspaceState.h b/Engine/Plugins/Experimental/Workspace/Source/WorkspaceEditor/Private/WorkspaceState.h
--- a/Engine/Plugins/Experimental/Workspace/Source/WorkspaceEditor/Private/WorkspaceState.h
+++ b/Engine/Plugins/Experimental/Workspace/Source/WorkspaceEditor/Private/WorkspaceState.h
@@ -20,7 +20,6 @@ class UWorkspaceState : public UObject
 	GENERATED_BODY()
 
 	UWorkspaceState();
-	virtual ~UWorkspaceState() override;
 
 	// Generate a file path used to persist workspace state
 	static FString GetWorkspacePath(const UWorkspace* InWorkspace);
