Editor crashes when opening Workspace for (AnimNext) assets
#fix pass through target WorkspaceFactory when opening AnimNext assets and call ConfigureProperties to cover other cases

diff --git a/Engine/Plugins/Experimental/AnimNext/Source/AnimNextEditor/Private/Graph/AnimNextGraphAssetDefinition.cpp b/Engine/Plugins/Experimental/AnimNext/Source/AnimNextEditor/Private/Graph/AnimNextGraphAssetDefinition.cpp
--- a/Engine/Plugins/Experimental/AnimNext/Source/AnimNextEditor/Private/Graph/AnimNextGraphAssetDefinition.cpp
+++ b/Engine/Plugins/Experimental/AnimNext/Source/AnimNextEditor/Private/Graph/AnimNextGraphAssetDefinition.cpp
@@ -4,6 +4,7 @@
 #include "AnimNextGraphEditor.h"
 #include "EditorCVars.h"
 #include "IWorkspaceEditorModule.h"
+#include "Workspace/AnimNextWorkspaceFactory.h"
 
 EAssetCommandResult UAssetDefinition_AnimNextGraph::OpenAssets(const FAssetOpenArgs& OpenArgs) const
 {
@@ -15,7 +16,7 @@ EAssetCommandResult UAssetDefinition_AnimNextGraph::OpenAssets(const FAssetOpenA
 		if(CVars::GUseWorkspaceEditor.GetValueOnGameThread())
 		{
 			IWorkspaceEditorModule& WorkspaceEditorModule = FModuleManager::Get().LoadModuleChecked<IWorkspaceEditorModule>("WorkspaceEditor");
-			WorkspaceEditorModule.OpenWorkspaceForObject(Asset, EOpenWorkspaceMethod::Default);
+			WorkspaceEditorModule.OpenWorkspaceForObject(Asset, EOpenWorkspaceMethod::Default, UAnimNextWorkspaceFactory::StaticClass());
 		}
 		else
 		{
diff --git a/Engine/Plugins/Experimental/AnimNext/Source/AnimNextEditor/Private/Scheduler/ScheduleAssetDefinition.cpp b/Engine/Plugins/Experimental/AnimNext/Source/AnimNextEditor/Private/Scheduler/ScheduleAssetDefinition.cpp
--- a/Engine/Plugins/Experimental/AnimNext/Source/AnimNextEditor/Private/Scheduler/ScheduleAssetDefinition.cpp
+++ b/Engine/Plugins/Experimental/AnimNext/Source/AnimNextEditor/Private/Scheduler/ScheduleAssetDefinition.cpp
@@ -4,6 +4,7 @@
 #include "Toolkits/SimpleAssetEditor.h"
 #include "EditorCVars.h"
 #include "IWorkspaceEditorModule.h"
+#include "Workspace/AnimNextWorkspaceFactory.h"
 
 #define LOCTEXT_NAMESPACE "AnimNextAssetDefinitions"
 
@@ -17,7 +18,7 @@ EAssetCommandResult UAssetDefinition_AnimNextSchedule::OpenAssets(const FAssetOp
 		if(CVars::GUseWorkspaceEditor.GetValueOnGameThread())
 		{
 			IWorkspaceEditorModule& WorkspaceEditorModule = FModuleManager::Get().LoadModuleChecked<IWorkspaceEditorModule>("WorkspaceEditor");
-			WorkspaceEditorModule.OpenWorkspaceForObject(Asset, EOpenWorkspaceMethod::Default);
+			WorkspaceEditorModule.OpenWorkspaceForObject(Asset, EOpenWorkspaceMethod::Default, UAnimNextWorkspaceFactory::StaticClass());
 		}
 		else
 		{
diff --git a/Engine/Plugins/Experimental/Workspace/Source/WorkspaceEditor/Private/WorkspaceEditorModule.cpp b/Engine/Plugins/Experimental/Workspace/Source/WorkspaceEditor/Private/WorkspaceEditorModule.cpp
--- a/Engine/Plugins/Experimental/Workspace/Source/WorkspaceEditor/Private/WorkspaceEditorModule.cpp
+++ b/Engine/Plugins/Experimental/Workspace/Source/WorkspaceEditor/Private/WorkspaceEditorModule.cpp
@@ -115,7 +115,7 @@ bool FWorkspaceEditorModule::GetExportedAssetsForWorkspace(const FAssetData& InW
 	return FWorkspaceAssetRegistryExports::StaticStruct()->ImportText(*TagValue, &OutExports, nullptr, PPF_None, nullptr, FWorkspaceAssetRegistryExports::StaticStruct()->GetName()) != nullptr;
 }
 
-void FWorkspaceEditorModule::OpenWorkspaceForObject(UObject* InObject, EOpenWorkspaceMethod InOpenMethod)
+void FWorkspaceEditorModule::OpenWorkspaceForObject(UObject* InObject, EOpenWorkspaceMethod InOpenMethod, const TSubclassOf<UWorkspaceFactory> WorkSpaceFactoryClass/*= UWorkspaceFactory::StaticClass()*/)
 {
 	FAssetRegistryModule& AssetRegistryModule = FModuleManager::LoadModuleChecked<FAssetRegistryModule>(TEXT("AssetRegistry"));
 
@@ -150,18 +150,21 @@ void FWorkspaceEditorModule::OpenWorkspaceForObject(UObject* InObject, EOpenWork
 
 	FWorkspaceEditor* WorkspaceEditor = nullptr;
 
-	auto HandleNewWorkspace = [InObject, &WorkspaceEditor]()
+	auto HandleNewWorkspace = [InObject, &WorkspaceEditor, WorkSpaceFactoryClass]()
 	{
-		UWorkspaceFactory* Factory = NewObject<UWorkspaceFactory>();
+		UWorkspaceFactory* Factory = NewObject<UWorkspaceFactory>(GetTransientPackage(), WorkSpaceFactoryClass.Get());
 		UPackage* Package = CreatePackage(nullptr);
 		FName PackageName = *FPaths::GetBaseFilename(Package->GetName());
-		UWorkspace* NewWorkspace = CastChecked<UWorkspace>(Factory->FactoryCreateNew(UWorkspace::StaticClass(), Package, PackageName, RF_Public | RF_Standalone, NULL, GWarn));
-		NewWorkspace->AddAsset(InObject, false);
-		NewWorkspace->MarkPackageDirty();
-		TSharedRef<FWorkspaceEditor> Editor = MakeShared<FWorkspaceEditor>();
-		Editor->InitEditor(EToolkitMode::Standalone, nullptr, NewWorkspace);
+		if (Factory->ConfigureProperties())
+		{
+			UWorkspace* NewWorkspace = CastChecked<UWorkspace>(Factory->FactoryCreateNew(UWorkspace::StaticClass(), Package, PackageName, RF_Public | RF_Standalone, NULL, GWarn));
+			NewWorkspace->AddAsset(InObject, false);
+			NewWorkspace->MarkPackageDirty();
+			TSharedRef<FWorkspaceEditor> Editor = MakeShared<FWorkspaceEditor>();
+			Editor->InitEditor(EToolkitMode::Standalone, nullptr, NewWorkspace);
 
-		WorkspaceEditor = &Editor.Get();
+			WorkspaceEditor = &Editor.Get();
+		}
 	};
 
 	auto HandleExistingWorkspace = [InObject, &WorkspaceEditor](const FAssetData& InAssetData)
diff --git a/Engine/Plugins/Experimental/Workspace/Source/WorkspaceEditor/Private/WorkspaceEditorModule.h b/Engine/Plugins/Experimental/Workspace/Source/WorkspaceEditor/Private/WorkspaceEditorModule.h
--- a/Engine/Plugins/Experimental/Workspace/Source/WorkspaceEditor/Private/WorkspaceEditorModule.h
+++ b/Engine/Plugins/Experimental/Workspace/Source/WorkspaceEditor/Private/WorkspaceEditorModule.h
@@ -23,7 +23,7 @@ private:
 	virtual void RegisterObjectDocumentType(const FTopLevelAssetPath& InClassPath, const FObjectDocumentArgs& InParams) override;
 	virtual void UnregisterObjectDocumentType(const FTopLevelAssetPath& InClassPath) override;
 	virtual FObjectDocumentArgs CreateGraphDocumentArgs(const FGraphDocumentWidgetArgs& InArgs) override;
-	virtual void OpenWorkspaceForObject(UObject* InObject, EOpenWorkspaceMethod InOpenMethod) override;
+	virtual void OpenWorkspaceForObject(UObject* InObject, EOpenWorkspaceMethod InOpenMethod, const TSubclassOf<UWorkspaceFactory> WorkSpaceFactoryClass) override;
 
 	// Find an existing registered object document type
 	const FObjectDocumentArgs* FindObjectDocumentType(const FTopLevelAssetPath& InClassPath) const;
diff --git a/Engine/Plugins/Experimental/Workspace/Source/WorkspaceEditor/Private/WorkspaceFactory.cpp b/Engine/Plugins/Experimental/Workspace/Source/WorkspaceEditor/Private/WorkspaceFactory.cpp
--- a/Engine/Plugins/Experimental/Workspace/Source/WorkspaceEditor/Private/WorkspaceFactory.cpp
+++ b/Engine/Plugins/Experimental/Workspace/Source/WorkspaceEditor/Private/WorkspaceFactory.cpp
@@ -56,6 +56,7 @@ bool UWorkspaceFactory::ConfigureProperties()
 UObject* UWorkspaceFactory::FactoryCreateNew(UClass* Class, UObject* InParent, FName Name, EObjectFlags Flags, UObject* Context, FFeedbackContext* Warn, FName CallingContext)
 {
 	UWorkspace* NewWorkspace = NewObject<UWorkspace>(InParent, Class, Name, Flags | RF_Public | RF_Standalone | RF_Transactional | RF_LoadCompleted);
+	ensureMsgf(SchemaClass != nullptr, TEXT("UWorkspaceFactory requires a valid Schema Class"));
 	NewWorkspace->SchemaClass = SchemaClass;
 	NewWorkspace->Guid = FGuid::NewGuid();
 
diff --git a/Engine/Plugins/Experimental/Workspace/Source/WorkspaceEditor/Public/IWorkspaceEditorModule.h b/Engine/Plugins/Experimental/Workspace/Source/WorkspaceEditor/Public/IWorkspaceEditorModule.h
--- a/Engine/Plugins/Experimental/Workspace/Source/WorkspaceEditor/Public/IWorkspaceEditorModule.h
+++ b/Engine/Plugins/Experimental/Workspace/Source/WorkspaceEditor/Public/IWorkspaceEditorModule.h
@@ -5,6 +5,7 @@
 #include "CoreMinimal.h"
 #include "GraphEditor.h"
 #include "InstancedStruct.h"
+#include "WorkspaceFactory.h"
 #include "Misc/Attribute.h"
 #include "Modules/ModuleInterface.h"
 
@@ -145,7 +146,7 @@ class IWorkspaceEditorModule : public IModuleInterface
 {
 public:
 	// Open an object inside a workspace editor.
-	virtual void OpenWorkspaceForObject(UObject* InObject, EOpenWorkspaceMethod InOpenMethod) = 0;
+	virtual void OpenWorkspaceForObject(UObject* InObject, EOpenWorkspaceMethod InOpenMethod, const TSubclassOf<UWorkspaceFactory> WorkSpaceFactoryClass = UWorkspaceFactory::StaticClass()) = 0;
 
 	// Register a widget factory method to spawn for a particular class
 	virtual void RegisterObjectDocumentType(const FTopLevelAssetPath& InClassPath, const FObjectDocumentArgs& InArgs) = 0;
