Workspace OFPA: created indirection by storing workspace entries as external objects (no contention on Workspace asset)

diff --git a/Engine/Plugins/Experimental/Workspace/Config/BaseWorkspace.ini b/Engine/Plugins/Experimental/Workspace/Config/BaseWorkspace.ini
new file mode 100644
--- /dev/null
+++ b/Engine/Plugins/Experimental/Workspace/Config/BaseWorkspace.ini
@@ -0,0 +1,2 @@
+[CoreRedirects]
++PropertyRedirects=(OldName="Workspace.Assets",NewName="Workspace.Assets_DEPRECATED")
\ No newline at end of file
diff --git a/Engine/Plugins/Experimental/Workspace/Source/WorkspaceEditor/Private/SWorkspaceView.cpp b/Engine/Plugins/Experimental/Workspace/Source/WorkspaceEditor/Private/SWorkspaceView.cpp
--- a/Engine/Plugins/Experimental/Workspace/Source/WorkspaceEditor/Private/SWorkspaceView.cpp
+++ b/Engine/Plugins/Experimental/Workspace/Source/WorkspaceEditor/Private/SWorkspaceView.cpp
@@ -44,7 +44,7 @@ void SWorkspaceView::Construct(const FArguments& InArgs, UWorkspace* InWorkspace
 	AssetPickerConfig.Filter = MakeARFilter();
 	AssetPickerConfig.OnShouldFilterAsset = FOnShouldFilterAsset::CreateLambda([this](const FAssetData& InAsset)
 	{
-		return !Workspace->Assets.Contains(TSoftObjectPtr<UObject>(InAsset.GetSoftObjectPath()));
+		return !Workspace->AssetEntries.ContainsByPredicate([InAsset](const UWorkspaceAssetEntry* Entry)-> bool { return Entry->Asset == TSoftObjectPtr<UObject>(InAsset.GetSoftObjectPath()); });
 	});
 	AssetPickerConfig.AssetShowWarningText = LOCTEXT("EmptyWorkspaceText", "Workspace is empty.\nDrag-drop to add assets to this workspace.");
 
diff --git a/Engine/Plugins/Experimental/Workspace/Source/WorkspaceEditor/Private/Workspace.cpp b/Engine/Plugins/Experimental/Workspace/Source/WorkspaceEditor/Private/Workspace.cpp
--- a/Engine/Plugins/Experimental/Workspace/Source/WorkspaceEditor/Private/Workspace.cpp
+++ b/Engine/Plugins/Experimental/Workspace/Source/WorkspaceEditor/Private/Workspace.cpp
@@ -2,6 +2,7 @@
 
 #include "Workspace.h"
 
+#include "ExternalPackageHelper.h"
 #include "WorkspaceSchema.h"
 #include "WorkspaceState.h"
 #include "AssetRegistry/AssetData.h"
@@ -9,6 +10,7 @@
 #include "UObject/ConstructorHelpers.h"
 #include "UObject/FortniteMainBranchObjectVersion.h"
 
+
 const FName UWorkspace::ExportsAssetRegistryTag = TEXT("Exports");
 
 bool UWorkspace::AddAsset(const FAssetData& InAsset, bool bSetupUndoRedo, bool bPrintPythonCommand)
@@ -24,7 +26,27 @@ bool UWorkspace::AddAsset(const FAssetData& InAsset, bool bSetupUndoRedo, bool b
 		Modify();
 	}
 
-	const int32 NewIndex = Assets.AddUnique(TSoftObjectPtr<UObject>(InAsset.GetSoftObjectPath()));
+	int32 NewIndex = INDEX_NONE; 
+	if (!AssetEntries.ContainsByPredicate([&InAsset](const UWorkspaceAssetEntry* AssetEntry)-> bool
+	{
+		return AssetEntry->Asset.ToSoftObjectPath() == InAsset.ToSoftObjectPath();
+	}))
+	{
+		// If we are a transient asset, dont use external packages
+		const UObject* Asset = InAsset.GetAsset();
+		check(Asset);
+		if(!Asset->HasAnyFlags(RF_Transient))
+		{
+			UWorkspaceAssetEntry* NewEntry = NewObject<UWorkspaceAssetEntry>(this, UWorkspaceAssetEntry::StaticClass(), NAME_None, RF_Transactional);
+			FExternalPackageHelper::SetPackagingMode(NewEntry, this, true, true, PKG_None);
+
+			NewEntry->Asset = TSoftObjectPtr<UObject>(InAsset.GetSoftObjectPath());
+			NewIndex = AssetEntries.Add(NewEntry);
+
+			NewEntry->MarkPackageDirty();
+		}
+	}
+	
 	if(NewIndex != INDEX_NONE)
 	{
 		BroadcastModified();
@@ -107,13 +129,21 @@ bool UWorkspace::RemoveAsset(const FAssetData& InAsset, bool bSetupUndoRedo, boo
 		Modify();
 	}
 
-	const int32 NumRemoved = Assets.Remove(TSoftObjectPtr<UObject>(InAsset.GetSoftObjectPath()));
-	if(NumRemoved > 0)
+	const int32 EntryIndex = AssetEntries.IndexOfByPredicate([&InAsset](const UWorkspaceAssetEntry* AssetEntry) -> bool
+	{
+		return AssetEntry->Asset.ToSoftObjectPath() == InAsset.ToSoftObjectPath();
+	});
+
+	if (EntryIndex != INDEX_NONE)
 	{
+		UWorkspaceAssetEntry* EntryToRemove = AssetEntries[EntryIndex];		
+		check(AssetEntries.Remove(EntryToRemove) == 1);
+		EntryToRemove->MarkAsGarbage();
+		
 		BroadcastModified();
 	}
 
-	return NumRemoved > 0;
+	return EntryIndex != INDEX_NONE;
 }
 
 bool UWorkspace::RemoveAsset(UObject* InAsset, bool bSetupUndoRedo, bool bPrintPythonCommand)
@@ -230,6 +260,21 @@ void UWorkspace::PostTransacted(const FTransactionObjectEvent& TransactionEvent)
 	BroadcastModified();
 }
 
+void UWorkspace::PostLoadExternalPackages()
+{
+	FExternalPackageHelper::LoadObjectsFromExternalPackages<UWorkspaceAssetEntry>(this, [this](UWorkspaceAssetEntry* InLoadedEntry)
+	{
+		check(IsValid(InLoadedEntry));
+		AssetEntries.Add(InLoadedEntry);
+	});
+}
+
+void UWorkspace::Serialize(FArchive& Ar)
+{
+	UObject::Serialize(Ar);
+	Ar.UsingCustomVersion(FFortniteMainBranchObjectVersion::GUID);
+}
+
 void UWorkspace::PostLoad()
 {
 	Super::PostLoad();
@@ -239,6 +284,28 @@ void UWorkspace::PostLoad()
 		Guid = FGuid::NewGuid();
 		SchemaClass = StaticLoadClass(UWorkspaceSchema::StaticClass(), nullptr, TEXT("/Script/AnimNextEditor.AnimNextWorkspaceSchema"));
 	}
+
+	if(GetLinkerCustomVersion(FFortniteMainBranchObjectVersion::GUID) < FFortniteMainBranchObjectVersion::AnimNextWorkspaceEntryConversion)
+	{
+		for (const TSoftObjectPtr<UObject>& SoftAsset : Assets_DEPRECATED)
+		{
+			if (!SoftAsset.IsNull())
+			{
+				UWorkspaceAssetEntry* NewEntry = NewObject<UWorkspaceAssetEntry>(this, UWorkspaceAssetEntry::StaticClass(), NAME_None, RF_Transactional);
+				FExternalPackageHelper::SetPackagingMode(NewEntry, this, true, true, PKG_None);
+
+				NewEntry->Asset = TSoftObjectPtr<UObject>(SoftAsset.ToSoftObjectPath());
+				AssetEntries.Add(NewEntry);				
+				NewEntry->GetPackage()->SetDirtyFlag(true);
+			}
+		}
+	
+		Assets_DEPRECATED.Empty();
+	}
+	else
+	{
+		PostLoadExternalPackages();		
+	}
 }
 
 void UWorkspace::PostDuplicate(bool bDuplicateForPIE)
@@ -253,11 +320,15 @@ void UWorkspace::GetAssetRegistryTags(FAssetRegistryTagsContext Context) const
 	Super::GetAssetRegistryTags(Context);
 
 	FWorkspaceAssetRegistryExports Exports;
-	Exports.Assets.Reserve(Assets.Num());
+	Exports.Assets.Reserve(AssetEntries.Num());
 
-	for(const TSoftObjectPtr<UObject>& Asset : Assets)
+	for(const UWorkspaceAssetEntry* AssetEntry : AssetEntries)
 	{
-		Exports.Assets.Emplace(Asset.GetUniqueID());
+		if (AssetEntry)
+		{
+			Exports.Assets.Emplace(AssetEntry->Asset.GetUniqueID());
+		}
+	
 	}
 
 	FString TagValue;
diff --git a/Engine/Plugins/Experimental/Workspace/Source/WorkspaceEditor/Private/Workspace.h b/Engine/Plugins/Experimental/Workspace/Source/WorkspaceEditor/Private/Workspace.h
--- a/Engine/Plugins/Experimental/Workspace/Source/WorkspaceEditor/Private/Workspace.h
+++ b/Engine/Plugins/Experimental/Workspace/Source/WorkspaceEditor/Private/Workspace.h
@@ -3,7 +3,9 @@
 #pragma once
 
 #include "CoreMinimal.h"
+#include "AssetDefinitionDefault.h"
 #include "Templates/SubclassOf.h"
+#include "WorkspaceAssetEntry.h"
 #include "Workspace.generated.h"
 
 struct FEditedDocumentInfo;
@@ -110,10 +112,18 @@ class UWorkspace : public UObject
 	virtual void GetAssetRegistryTags(FAssetRegistryTagsContext Context) const override;
 	virtual void PostTransacted(const FTransactionObjectEvent& TransactionEvent) override;
 	virtual bool IsEditorOnly() const override { return true; }
+	virtual void Serialize(FArchive& Ar) override;
 
+protected:
+	void PostLoadExternalPackages();
 	// All of the assets referenced by this workspace
 	UPROPERTY()
-	TArray<TSoftObjectPtr<UObject>> Assets;
+	TArray<TSoftObjectPtr<UObject>> Assets_DEPRECATED;
+
+public:
+	// All of the assets in-directly referenced by this workspace
+	UPROPERTY(transient)
+	TArray<TObjectPtr<UWorkspaceAssetEntry>> AssetEntries;
 
 	// Schema for this workspace
 	UPROPERTY(AssetRegistrySearchable)
diff --git a/Engine/Plugins/Experimental/Workspace/Source/WorkspaceEditor/Private/WorkspaceAssetEntry.cpp b/Engine/Plugins/Experimental/Workspace/Source/WorkspaceEditor/Private/WorkspaceAssetEntry.cpp
new file mode 100644
--- /dev/null
+++ b/Engine/Plugins/Experimental/Workspace/Source/WorkspaceEditor/Private/WorkspaceAssetEntry.cpp
@@ -0,0 +1,11 @@
+﻿// Copyright Epic Games, Inc. All Rights Reserved.
+
+#include "WorkspaceAssetEntry.h"
+#include "UObject/Package.h"
+
+bool UWorkspaceAssetEntry::IsAsset() const
+{
+	// Entries are considered assets to allow using the asset logic for save dialogs, etc.
+	// Also, they return true even if pending kill, in order to show up as deleted in these dialogs.
+	return IsPackageExternal() && !GetPackage()->HasAnyFlags(RF_Transient) && !HasAnyFlags(RF_Transient | RF_ClassDefaultObject);
+}
diff --git a/Engine/Plugins/Experimental/Workspace/Source/WorkspaceEditor/Private/WorkspaceAssetEntry.h b/Engine/Plugins/Experimental/Workspace/Source/WorkspaceEditor/Private/WorkspaceAssetEntry.h
new file mode 100644
--- /dev/null
+++ b/Engine/Plugins/Experimental/Workspace/Source/WorkspaceEditor/Private/WorkspaceAssetEntry.h
@@ -0,0 +1,18 @@
+﻿// Copyright Epic Games, Inc. All Rights Reserved.
+
+#pragma once
+
+#include "CoreMinimal.h"
+#include "UObject/Object.h"
+#include "WorkspaceAssetEntry.generated.h"
+
+UCLASS()
+class UWorkspaceAssetEntry : public UObject
+{
+	GENERATED_BODY()
+public:	
+	virtual bool IsAsset() const override;
+
+	UPROPERTY()
+	TSoftObjectPtr<UObject> Asset;	
+};
diff --git a/Engine/Plugins/Experimental/Workspace/Source/WorkspaceEditor/Private/WorkspaceAssetEntryAssetDefinition.cpp b/Engine/Plugins/Experimental/Workspace/Source/WorkspaceEditor/Private/WorkspaceAssetEntryAssetDefinition.cpp
new file mode 100644
--- /dev/null
+++ b/Engine/Plugins/Experimental/Workspace/Source/WorkspaceEditor/Private/WorkspaceAssetEntryAssetDefinition.cpp
@@ -0,0 +1,29 @@
+﻿// Copyright Epic Games, Inc. All Rights Reserved.
+
+#include "WorkspaceAssetEntryAssetDefinition.h"
+#include "WorkspaceAssetEntry.h"
+
+#define LOCTEXT_NAMESPACE "AssetDefinition_WorkspaceAssetEntry"
+
+FText UAssetDefinition_WorkspaceAssetEntry::GetAssetDisplayName() const
+{
+	return LOCTEXT("WorkSpaceAssetEntry", "Workspace Asset Entry");
+}
+
+FLinearColor UAssetDefinition_WorkspaceAssetEntry::GetAssetColor() const
+{
+	return FLinearColor(FColor(64,64,64));
+}
+
+TSoftClassPtr<> UAssetDefinition_WorkspaceAssetEntry::GetAssetClass() const
+{
+	return UWorkspaceAssetEntry::StaticClass();
+}
+
+FText UAssetDefinition_WorkspaceAssetEntry::GetObjectDisplayNameText(UObject* Object) const
+{
+	const UWorkspaceAssetEntry* Entry = CastChecked<UWorkspaceAssetEntry>(Object);
+	return FText::FromString(Entry->Asset.GetAssetName());
+}
+
+#undef LOCTEXT_NAMESPACE // "AssetDefinition_WorkspaceAssetEntry"
\ No newline at end of file
diff --git a/Engine/Plugins/Experimental/Workspace/Source/WorkspaceEditor/Private/WorkspaceAssetEntryAssetDefinition.h b/Engine/Plugins/Experimental/Workspace/Source/WorkspaceEditor/Private/WorkspaceAssetEntryAssetDefinition.h
new file mode 100644
--- /dev/null
+++ b/Engine/Plugins/Experimental/Workspace/Source/WorkspaceEditor/Private/WorkspaceAssetEntryAssetDefinition.h
@@ -0,0 +1,21 @@
+﻿// Copyright Epic Games, Inc. All Rights Reserved.
+
+#pragma once
+
+#include "CoreMinimal.h"
+#include "AssetDefinitionDefault.h"
+#include "UObject/Object.h"
+#include "WorkspaceAssetEntryAssetDefinition.generated.h"
+
+UCLASS()
+class UAssetDefinition_WorkspaceAssetEntry  : public UAssetDefinitionDefault
+{
+	GENERATED_BODY()
+
+public:
+	// UAssetDefinition interface
+	virtual FText GetAssetDisplayName() const override;
+	virtual FLinearColor GetAssetColor() const override;
+	virtual TSoftClassPtr<UObject> GetAssetClass() const override;
+	virtual FText GetObjectDisplayNameText(UObject* Object) const override;
+};
diff --git a/Engine/Plugins/Experimental/Workspace/Source/WorkspaceEditor/Private/WorkspaceEditor.cpp b/Engine/Plugins/Experimental/Workspace/Source/WorkspaceEditor/Private/WorkspaceEditor.cpp
--- a/Engine/Plugins/Experimental/Workspace/Source/WorkspaceEditor/Private/WorkspaceEditor.cpp
+++ b/Engine/Plugins/Experimental/Workspace/Source/WorkspaceEditor/Private/WorkspaceEditor.cpp
@@ -270,9 +270,9 @@ void FWorkspaceEditor::GetSaveableObjects(TArray<UObject*>& OutObjects) const
 
 	if(!bSavingWorkspaceOnly)
 	{
-		for(TSoftObjectPtr<UObject>& SoftAsset : Workspace->Assets)
+		for(const UWorkspaceAssetEntry* Entry : Workspace->AssetEntries)
 		{
-			if(UObject* Asset = SoftAsset.Get())
+			if(UObject* Asset = Entry->Asset.Get())
 			{
 				// Add object referenced by workspace
 				OutObjects.Add(Asset);
diff --git a/Engine/Source/Runtime/Core/Public/UObject/FortniteMainBranchObjectVersion.h b/Engine/Source/Runtime/Core/Public/UObject/FortniteMainBranchObjectVersion.h
--- a/Engine/Source/Runtime/Core/Public/UObject/FortniteMainBranchObjectVersion.h
+++ b/Engine/Source/Runtime/Core/Public/UObject/FortniteMainBranchObjectVersion.h
@@ -435,6 +435,9 @@ struct FFortniteMainBranchObjectVersion
 		// Introduce new structure for customizing the landscape edit layer behavior
 		IntroduceLandscapeEditLayerClass,
 
+		// Change workspaces to store asset references as external objects
+		AnimNextWorkspaceEntryConversion,
+
 		// -----<new versions can be added above this line>-------------------------------------------------
 		VersionPlusOne,
 		LatestVersion = VersionPlusOne - 1
