Workspace refactor
- Rebase IWorkspaceEditor to FBaseAssetToolkit
- Introduced WorkspaceAssetEditor
- Removed Details/Workspace tab summoners (former is created by FBaseAssetToolkit and latter is handled by new WorkSpaceEditor implementation)
- Removed EditorMode
- Moved mutating code from Workspace to EditorSubsystem
- Slightly tweaked upgrade path for Workspace (reducing code duplication)
- Added details-customization callbacks/registering to Workspace module

diff --git a/Engine/Plugins/Experimental/Workspace/Source/WorkspaceEditor/Private/DetailsTabSummoner.cpp b/Engine/Plugins/Experimental/Workspace/Source/WorkspaceEditor/Private/DetailsTabSummoner.cpp
deleted file mode 100644
diff --git a/Engine/Plugins/Experimental/Workspace/Source/WorkspaceEditor/Private/DetailsTabSummoner.h b/Engine/Plugins/Experimental/Workspace/Source/WorkspaceEditor/Private/DetailsTabSummoner.h
deleted file mode 100644
diff --git a/Engine/Plugins/Experimental/Workspace/Source/WorkspaceEditor/Private/Workspace.cpp b/Engine/Plugins/Experimental/Workspace/Source/WorkspaceEditor/Private/Workspace.cpp
--- a/Engine/Plugins/Experimental/Workspace/Source/WorkspaceEditor/Private/Workspace.cpp
+++ b/Engine/Plugins/Experimental/Workspace/Source/WorkspaceEditor/Private/Workspace.cpp
@@ -10,7 +10,6 @@
 #include "UObject/ConstructorHelpers.h"
 #include "UObject/FortniteMainBranchObjectVersion.h"
 
-
 const FName UWorkspace::ExportsAssetRegistryTag = TEXT("Exports");
 
 bool UWorkspace::AddAsset(const FAssetData& InAsset, bool bSetupUndoRedo, bool bPrintPythonCommand)
@@ -209,7 +208,7 @@ bool UWorkspace::RemoveAssets(TConstArrayView<FAssetData> InAssets, bool bSetupU
 
 bool UWorkspace::IsAssetSupported(const FAssetData& InAsset)
 {
-	TConstArrayView<FTopLevelAssetPath> SupportedAssets = GetSchema()->GetSupportedAssetClassPaths();
+	const TConstArrayView<FTopLevelAssetPath> SupportedAssets = GetSchema()->GetSupportedAssetClassPaths();
 	return SupportedAssets.IsEmpty() || SupportedAssets.Contains(InAsset.AssetClassPath);
 }
 
@@ -289,17 +288,17 @@ void UWorkspace::PostLoad()
 	{
 		for (const TSoftObjectPtr<UObject>& SoftAsset : Assets_DEPRECATED)
 		{
-			if (!SoftAsset.IsNull())
+			if (UObject* Asset = SoftAsset.LoadSynchronous())
 			{
-				UWorkspaceAssetEntry* NewEntry = NewObject<UWorkspaceAssetEntry>(this, UWorkspaceAssetEntry::StaticClass(), NAME_None, RF_Transactional);
-				FExternalPackageHelper::SetPackagingMode(NewEntry, this, true, true, PKG_None);
-
-				NewEntry->Asset = TSoftObjectPtr<UObject>(SoftAsset.ToSoftObjectPath());
-				AssetEntries.Add(NewEntry);				
-				NewEntry->GetPackage()->SetDirtyFlag(true);
+				AddAsset(Asset);
 			}
 		}
-	
+
+		for (const TObjectPtr<UWorkspaceAssetEntry>& Entry : AssetEntries)
+		{
+			Entry->Asset->GetPackage()->SetDirtyFlag(true);	
+		}
+
 		Assets_DEPRECATED.Empty();
 	}
 	else
@@ -335,4 +334,4 @@ void UWorkspace::GetAssetRegistryTags(FAssetRegistryTagsContext Context) const
 	FWorkspaceAssetRegistryExports::StaticStruct()->ExportText(TagValue, &Exports, nullptr, nullptr, PPF_None, nullptr);
 
 	Context.AddTag(FAssetRegistryTag(ExportsAssetRegistryTag, TagValue, FAssetRegistryTag::TT_Hidden));
-}
+}
\ No newline at end of file
diff --git a/Engine/Plugins/Experimental/Workspace/Source/WorkspaceEditor/Private/Workspace.h b/Engine/Plugins/Experimental/Workspace/Source/WorkspaceEditor/Private/Workspace.h
--- a/Engine/Plugins/Experimental/Workspace/Source/WorkspaceEditor/Private/Workspace.h
+++ b/Engine/Plugins/Experimental/Workspace/Source/WorkspaceEditor/Private/Workspace.h
@@ -97,9 +97,7 @@ class UWorkspace : public UObject
 	UWorkspaceSchema* GetSchema() const;
 
 	void LoadState();
-
 	void SaveState();
-
 	UWorkspaceState* GetState() const;
 
 	void ReportError(const TCHAR* InMessage) const;
@@ -114,13 +112,12 @@ class UWorkspace : public UObject
 	virtual bool IsEditorOnly() const override { return true; }
 	virtual void Serialize(FArchive& Ar) override;
 
-protected:
 	void PostLoadExternalPackages();
+	
 	// All of the assets referenced by this workspace
 	UPROPERTY()
 	TArray<TSoftObjectPtr<UObject>> Assets_DEPRECATED;
 
-public:
 	// All of the assets in-directly referenced by this workspace
 	UPROPERTY(transient)
 	TArray<TObjectPtr<UWorkspaceAssetEntry>> AssetEntries;
diff --git a/Engine/Plugins/Experimental/Workspace/Source/WorkspaceEditor/Private/WorkspaceAssetDefinition.cpp b/Engine/Plugins/Experimental/Workspace/Source/WorkspaceEditor/Private/WorkspaceAssetDefinition.cpp
--- a/Engine/Plugins/Experimental/Workspace/Source/WorkspaceEditor/Private/WorkspaceAssetDefinition.cpp
+++ b/Engine/Plugins/Experimental/Workspace/Source/WorkspaceEditor/Private/WorkspaceAssetDefinition.cpp
@@ -3,6 +3,7 @@
 #include "WorkspaceAssetDefinition.h"
 #include "WorkspaceEditor.h"
 #include "WorkspaceSchema.h"
+#include "WorkspaceAssetEditor.h"
 
 #define LOCTEXT_NAMESPACE "AssetDefinition_Workspace"
 
@@ -33,8 +34,10 @@ EAssetCommandResult UAssetDefinition_Workspace::OpenAssets(const FAssetOpenArgs&
 {
 	for (UWorkspace* Asset : OpenArgs.LoadObjects<UWorkspace>())
 	{
-		TSharedRef<UE::Workspace::FWorkspaceEditor> Editor = MakeShared<UE::Workspace::FWorkspaceEditor>();
-		Editor->InitEditor(OpenArgs.GetToolkitMode(), OpenArgs.ToolkitHost, Asset);
+		UAssetEditorSubsystem* AssetEditorSubsystem = GEditor->GetEditorSubsystem<UAssetEditorSubsystem>();
+		UWorkspaceAssetEditor* AssetEditor = NewObject<UWorkspaceAssetEditor>(AssetEditorSubsystem, NAME_None, RF_Transient);
+		AssetEditor->SetObjectToEdit(Asset);
+		AssetEditor->Initialize();
 	}
 
 	return EAssetCommandResult::Handled;
diff --git a/Engine/Plugins/Experimental/Workspace/Source/WorkspaceEditor/Private/WorkspaceAssetEditor.cpp b/Engine/Plugins/Experimental/Workspace/Source/WorkspaceEditor/Private/WorkspaceAssetEditor.cpp
new file mode 100644
--- /dev/null
+++ b/Engine/Plugins/Experimental/Workspace/Source/WorkspaceEditor/Private/WorkspaceAssetEditor.cpp
@@ -0,0 +1,27 @@
+// Copyright Epic Games, Inc. All Rights Reserved.
+
+#include "WorkspaceAssetEditor.h"
+#include "WorkspaceEditor.h"
+
+#include UE_INLINE_GENERATED_CPP_BY_NAME(WorkspaceAssetEditor)
+
+void UWorkspaceAssetEditor::GetObjectsToEdit(TArray<UObject*>& OutObjectsToEdit)
+{
+	OutObjectsToEdit.Add(ObjectToEdit);
+}
+
+void UWorkspaceAssetEditor::SetObjectToEdit(UWorkspace* InWorkspace)
+{
+	ObjectToEdit = InWorkspace;
+}
+
+UWorkspace* UWorkspaceAssetEditor::GetObjectToEdit()
+{
+	return ObjectToEdit;
+}
+
+TSharedPtr<FBaseAssetToolkit> UWorkspaceAssetEditor::CreateToolkit()
+{
+	return MakeShared<UE::Workspace::FWorkspaceEditor>(this);
+}
+
diff --git a/Engine/Plugins/Experimental/Workspace/Source/WorkspaceEditor/Private/WorkspaceAssetEditor.h b/Engine/Plugins/Experimental/Workspace/Source/WorkspaceEditor/Private/WorkspaceAssetEditor.h
new file mode 100644
--- /dev/null
+++ b/Engine/Plugins/Experimental/Workspace/Source/WorkspaceEditor/Private/WorkspaceAssetEditor.h
@@ -0,0 +1,25 @@
+// Copyright Epic Games, Inc. All Rights Reserved.
+
+#pragma once
+
+#include "Tools/UAssetEditor.h"
+#include "Workspace.h"
+#include "WorkspaceAssetEditor.generated.h"
+
+class FBaseAssetToolkit;
+
+UCLASS(Transient)
+class UWorkspaceAssetEditor : public UAssetEditor
+{
+	GENERATED_BODY()
+public:
+	void SetObjectToEdit(UWorkspace* InWorkspace);
+	UWorkspace* GetObjectToEdit();
+protected:
+	virtual void GetObjectsToEdit(TArray<UObject*>& OutObjectsToEdit) override;
+	virtual TSharedPtr<FBaseAssetToolkit> CreateToolkit() override;
+
+private:
+	TObjectPtr<UWorkspace> ObjectToEdit;
+};
+
diff --git a/Engine/Plugins/Experimental/Workspace/Source/WorkspaceEditor/Private/WorkspaceEditor.cpp b/Engine/Plugins/Experimental/Workspace/Source/WorkspaceEditor/Private/WorkspaceEditor.cpp
--- a/Engine/Plugins/Experimental/Workspace/Source/WorkspaceEditor/Private/WorkspaceEditor.cpp
+++ b/Engine/Plugins/Experimental/Workspace/Source/WorkspaceEditor/Private/WorkspaceEditor.cpp
@@ -2,17 +2,16 @@
 
 #include "WorkspaceEditor.h"
 
-#include "WorkspaceEditorMode.h"
 #include "Workspace.h"
 #include "AssetDocumentSummoner.h"
 #include "ExternalPackageHelper.h"
 #include "Widgets/Docking/SDockTab.h"
 #include "WorkflowOrientedApp/WorkflowUObjectDocuments.h"
-#include "SWorkspacePicker.h"
+#include "WorkspaceAssetEditor.h"
 #include "WorkspaceState.h"
 #include "WorkspaceDocumentState.h"
 #include "WorkspaceEditorModule.h"
-#include "Engine/Blueprint.h"
+#include "SWorkspaceView.h"
 
 #define LOCTEXT_NAMESPACE "WorkspaceEditor"
 
@@ -35,46 +34,90 @@ namespace WorkspaceTabs
 
 const FName WorkspaceAppIdentifier("WorkspaceEditor");
 
-FWorkspaceEditor::FWorkspaceEditor()
+FWorkspaceEditor::FWorkspaceEditor(UWorkspaceAssetEditor* InOwningAssetEditor) : IWorkspaceEditor(InOwningAssetEditor)
 {
+	Workspace = Cast<UWorkspaceAssetEditor>(OwningAssetEditor)->GetObjectToEdit();
 }
 
-FWorkspaceEditor::~FWorkspaceEditor()
+void FWorkspaceEditor::CreateWidgets()
 {
-}
-
-void FWorkspaceEditor::InitEditor(const EToolkitMode::Type InMode, const TSharedPtr<IToolkitHost>& InInitToolkitHost, UWorkspace* InWorkspace)
-{
-	Workspace = InWorkspace;
-
-	Workspace->LoadState();
-
 	DocumentManager = MakeShared<FDocumentTracker>(NAME_None);
 	DocumentManager->Initialize(SharedThis(this));
+	
+	FBaseAssetToolkit::CreateWidgets();
 
-	FWorkspaceEditorModule& WorkspaceEditorModule = FModuleManager::Get().LoadModuleChecked<FWorkspaceEditorModule>("WorkspaceEditor");
+	const FWorkspaceEditorModule& WorkspaceEditorModule = FModuleManager::Get().LoadModuleChecked<FWorkspaceEditorModule>("WorkspaceEditor");
 
 	// Build document summoners for each workspace layout area
-	TSharedRef<FAssetDocumentSummoner> LeftAssetDocumentSummoner = MakeShared<FAssetDocumentSummoner>(WorkspaceTabs::LeftDocumentArea, SharedThis(this));
+	const TSharedRef<FAssetDocumentSummoner> LeftAssetDocumentSummoner = MakeShared<FAssetDocumentSummoner>(WorkspaceTabs::LeftDocumentArea, SharedThis(this));
 	LeftAssetDocumentSummoner->SetAllowedClassPaths(WorkspaceEditorModule.GetAllowedObjectTypesForArea(WorkspaceTabs::LeftDocumentArea));
 	DocumentManager->RegisterDocumentFactory(LeftAssetDocumentSummoner);
 
-	TSharedRef<FAssetDocumentSummoner> MiddleAssetDocumentSummoner = MakeShared<FAssetDocumentSummoner>(WorkspaceTabs::MiddleDocumentArea, SharedThis(this));
+	const TSharedRef<FAssetDocumentSummoner> MiddleAssetDocumentSummoner = MakeShared<FAssetDocumentSummoner>(WorkspaceTabs::MiddleDocumentArea, SharedThis(this));
 	MiddleAssetDocumentSummoner->SetAllowedClassPaths(WorkspaceEditorModule.GetAllowedObjectTypesForArea(WorkspaceTabs::MiddleDocumentArea));
 	DocumentManager->RegisterDocumentFactory(MiddleAssetDocumentSummoner);
 
-	TSharedRef<FAssetDocumentSummoner> RightAssetDocumentSummoner = MakeShared<FAssetDocumentSummoner>(WorkspaceTabs::RightDocumentArea, SharedThis(this));
+	const TSharedRef<FAssetDocumentSummoner> RightAssetDocumentSummoner = MakeShared<FAssetDocumentSummoner>(WorkspaceTabs::RightDocumentArea, SharedThis(this));
 	RightAssetDocumentSummoner->SetAllowedClassPaths(WorkspaceEditorModule.GetAllowedObjectTypesForArea(WorkspaceTabs::RightDocumentArea));
 	DocumentManager->RegisterDocumentFactory(RightAssetDocumentSummoner);
+	
+	check(DetailsView.IsValid());
+	WorkspaceEditorModule.ApplyWorkspaceDetailsCustomization(DetailsView);
+
+	StandaloneDefaultLayout = FTabManager::NewLayout("Standalone_WorkspaceEditor_Layout_v1.1")
+	->AddArea
+	(
+		FTabManager::NewPrimaryArea()
+		->SetOrientation(Orient_Vertical)
+		->Split
+		(
+			FTabManager::NewSplitter()
+			->SetSizeCoefficient(1.0f)
+			->SetOrientation(Orient_Horizontal)
+			->Split
+			(
+				FTabManager::NewStack()
+				->SetSizeCoefficient(0.25f)
+				->SetHideTabWell(false)
+				->AddTab(WorkspaceTabs::WorkspaceView, ETabState::OpenedTab)
+				->AddTab(WorkspaceTabs::LeftDocumentArea, ETabState::ClosedTab)
+			)
+			->Split
+			(
+				FTabManager::NewStack()
+				->SetSizeCoefficient(0.5f)
+				->SetHideTabWell(false)
+				->AddTab(WorkspaceTabs::MiddleDocumentArea, ETabState::ClosedTab)
+			)
+			->Split
+			(
+				FTabManager::NewStack()
+				->SetSizeCoefficient(0.25f)
+				->SetHideTabWell(false)
+				->AddTab(WorkspaceTabs::RightDocumentArea, ETabState::ClosedTab)
+				->AddTab(FBaseAssetToolkit::DetailsTabID, ETabState::OpenedTab)
+			)
+		)
+	);
+
+	TWeakPtr<FWorkspaceEditor> WeakToolkit = StaticCastSharedRef<FWorkspaceEditor>(AsShared());
+	WorkspaceView = SNew(SWorkspaceView, Workspace)
+		.OnAssetsOpened_Lambda([WeakToolkit](TConstArrayView<FAssetData> InAssets)
+		{
+			if(const TSharedPtr<FWorkspaceEditor> HostingApp = WeakToolkit.Pin())
+			{
+				HostingApp->OpenAssets(InAssets);
+			}
+		});
 
-	constexpr bool bCreateDefaultStandaloneMenu = true;
-	constexpr bool bCreateDefaultToolbar = true;
-	InitAssetEditor(InMode, InInitToolkitHost, WorkspaceAppIdentifier, FTabManager::FLayout::NullLayout, bCreateDefaultStandaloneMenu, bCreateDefaultToolbar, InWorkspace);
-
+	
 	BindCommands();
+}
 
-	AddApplicationMode(WorkspaceModes::WorkspaceEditor, MakeShared<FWorkspaceEditorMode>(SharedThis(this)));
-	SetCurrentMode(WorkspaceModes::WorkspaceEditor);
+void FWorkspaceEditor::PostInitAssetEditor()
+{
+	Workspace->LoadState();
+	RestoreEditedObjectState();
 
 	ExtendMenu();
 	ExtendToolbar();
@@ -88,7 +131,7 @@ void FWorkspaceEditor::RestoreEditedObjectState()
 	{
 		if (UObject* Object = DocumentState.Get().Object.TryLoad())
 		{
-			if(TSharedPtr<SDockTab> DockTab = OpenDocument(Object, FDocumentTracker::RestorePreviousDocument))
+			if(const TSharedPtr<SDockTab> DockTab = OpenDocument(Object, FDocumentTracker::RestorePreviousDocument))
 			{
 				FWorkspaceEditorModule& WorkspaceEditorModule = FModuleManager::LoadModuleChecked<FWorkspaceEditorModule>("WorkspaceEditor");
 				const FObjectDocumentArgs* DocumentArgs = WorkspaceEditorModule.FindObjectDocumentType(Object->GetClass()->GetClassPathName());
@@ -101,7 +144,7 @@ void FWorkspaceEditor::RestoreEditedObjectState()
 	}
 }
 
-void FWorkspaceEditor::SaveEditedObjectState()
+void FWorkspaceEditor::SaveEditedObjectState() const
 {
 	// Clear edited document state
 	UWorkspaceState* State = Workspace->GetState();
@@ -116,12 +159,12 @@ void FWorkspaceEditor::SaveEditedObjectState()
 
 TSharedPtr<SDockTab> FWorkspaceEditor::OpenDocument(const UObject* InForObject, FDocumentTracker::EOpenDocumentCause InCause)
 {
-	TSharedRef<FTabPayload_UObject> Payload = FTabPayload_UObject::Make(InForObject);
+	const TSharedRef<FTabPayload_UObject> Payload = FTabPayload_UObject::Make(InForObject);
 	TSharedPtr<SDockTab> NewTab = DocumentManager->OpenDocument(Payload, InCause);
 
 	if(InCause != FDocumentTracker::RestorePreviousDocument)
 	{
-		FWorkspaceEditorModule& WorkspaceEditorModule = FModuleManager::LoadModuleChecked<FWorkspaceEditorModule>("WorkspaceEditor");
+		const FWorkspaceEditorModule& WorkspaceEditorModule = FModuleManager::LoadModuleChecked<FWorkspaceEditorModule>("WorkspaceEditor");
 		const FObjectDocumentArgs* DocumentArgs = WorkspaceEditorModule.FindObjectDocumentType(InForObject->GetClass()->GetClassPathName());
 		if(DocumentArgs != nullptr && DocumentArgs->OnGetDocumentState.IsBound())
 		{
@@ -140,7 +183,7 @@ void FWorkspaceEditor::OpenAssets(TConstArrayView<FAssetData> InAssets)
 {
 	for(const FAssetData& Asset : InAssets)
 	{
-		if(UObject* LoadedAsset = Asset.GetAsset())
+		if(const UObject* LoadedAsset = Asset.GetAsset())
 		{
 			OpenDocument(LoadedAsset, FDocumentTracker::EOpenDocumentCause::OpenNewDocument);
 		}
@@ -149,7 +192,7 @@ void FWorkspaceEditor::OpenAssets(TConstArrayView<FAssetData> InAssets)
 
 void FWorkspaceEditor::OpenObjects(TConstArrayView<UObject*> InObjects)
 {
-	for(UObject* Object : InObjects)
+	for(const UObject* Object : InObjects)
 	{
 		OpenDocument(Object, FDocumentTracker::EOpenDocumentCause::OpenNewDocument);
 	}
@@ -159,7 +202,7 @@ void FWorkspaceEditor::CloseObjects(TConstArrayView<UObject*> InObjects)
 {
 	if(InObjects.Num() > 0)
 	{
-		for(UObject* Object : InObjects)
+		for(const UObject* Object : InObjects)
 		{
 			CloseDocumentTab(Object);
 		}
@@ -198,14 +241,28 @@ void FWorkspaceEditor::ExtendToolbar()
 
 void FWorkspaceEditor::RegisterTabSpawners(const TSharedRef<FTabManager>& InTabManager)
 {
-	DocumentManager->SetTabManager(InTabManager);
+	FBaseAssetToolkit::RegisterTabSpawners(InTabManager);
+
+	InTabManager->RegisterTabSpawner(WorkspaceTabs::WorkspaceView, FOnSpawnTab::CreateLambda([this](const FSpawnTabArgs& Args) -> TSharedRef<SDockTab>
+	{
+		check(Args.GetTabId() == WorkspaceTabs::WorkspaceView);
+		
+		return SNew(SDockTab)
+			.Label(LOCTEXT("WorkspaceTabLabel", "Workspace"))
+			[
+				WorkspaceView.ToSharedRef()
+			];
+	}))
+	.SetDisplayName(LOCTEXT("WorkspaceTabLabel", "Workspace"))
+	.SetIcon(FSlateIcon("EditorStyle", "LevelEditor.Tabs.Outliner"))
+	.SetTooltipText(LOCTEXT("WorkspaceTabToolTip", "Shows the workspace outliner tab."));
 
-	FWorkflowCentricApplication::RegisterTabSpawners(InTabManager);
+	DocumentManager->SetTabManager(InTabManager);
 }
 
 void FWorkspaceEditor::UnregisterTabSpawners(const TSharedRef<FTabManager>& InTabManager)
 {
-	FWorkflowCentricApplication::UnregisterTabSpawners(InTabManager);
+	FBaseAssetToolkit::UnregisterTabSpawners(InTabManager);
 }
 
 FName FWorkspaceEditor::GetToolkitFName() const
@@ -235,26 +292,26 @@ void FWorkspaceEditor::InitToolMenuContext(FToolMenuContext& InMenuContext)
 void FWorkspaceEditor::SaveAsset_Execute()
 {
 	// If asset is a default 'Untitled' workspace, redirect to the 'save as' flow
-	FString AssetPath = Workspace->GetOutermost()->GetPathName();
+	const FString AssetPath = Workspace->GetOutermost()->GetPathName();
 	if(AssetPath.StartsWith(TEXT("/Temp/Untitled")))
 	{
-		// Ensure we dont also 'save as' other externally linked assets at this point
+		// Ensure we do not also 'save as' other externally linked assets at this point
 		TGuardValue<bool> SaveWorkspaceOnly(bSavingWorkspaceOnly, true);
 
 		SaveAssetAs_Execute();
 	}
 	else
 	{
-		FWorkflowCentricApplication::SaveAsset_Execute();
+		FBaseAssetToolkit::SaveAsset_Execute();
 	}
 }
 
-void FWorkspaceEditor::CloseDocumentTab(const UObject* DocumentID)
+void FWorkspaceEditor::CloseDocumentTab(const UObject* DocumentID) const
 {
 	UWorkspaceState* State = Workspace->GetState();
 	State->DocumentStates.Remove(TInstancedStruct<FWorkspaceDocumentState>::Make(DocumentID));
 
-	TSharedRef<FTabPayload_UObject> Payload = FTabPayload_UObject::Make(DocumentID);
+	const TSharedRef<FTabPayload_UObject> Payload = FTabPayload_UObject::Make(DocumentID);
 	DocumentManager->CloseTab(Payload);
 }
 
@@ -266,8 +323,14 @@ bool FWorkspaceEditor::InEditingMode() const
 void FWorkspaceEditor::GetSaveableObjects(TArray<UObject*>& OutObjects) const
 {
 	// Base class will pick up edited object
-	FWorkflowCentricApplication::GetSaveableObjects(OutObjects);
+	FBaseAssetToolkit::GetSaveableObjects(OutObjects);
 
+	for (UObject* Object : GetEditingObjects())
+	{
+		// Get external objects too
+		FExternalPackageHelper::GetExternalSaveableObjects(Object, OutObjects);	
+	}
+	
 	if(!bSavingWorkspaceOnly)
 	{
 		for(const UWorkspaceAssetEntry* Entry : Workspace->AssetEntries)
@@ -284,7 +347,7 @@ void FWorkspaceEditor::GetSaveableObjects(TArray<UObject*>& OutObjects) const
 	}
 }
 
-void FWorkspaceEditor::RecordDocumentState(const TInstancedStruct<FWorkspaceDocumentState>& InState)
+void FWorkspaceEditor::RecordDocumentState(const TInstancedStruct<FWorkspaceDocumentState>& InState) const
 {
 	UWorkspaceState* State = Workspace->GetState();
 	State->DocumentStates.AddUnique(InState);
@@ -294,12 +357,12 @@ bool FWorkspaceEditor::OnRequestClose(EAssetEditorCloseReason InCloseReason)
 {
 	auto RequiresSave = [this]()
 	{
-		UPackage* Package = Workspace->GetOutermost();
+		const UPackage* Package = Workspace->GetOutermost();
 		return Package->GetPathName().StartsWith(TEXT("/Temp/Untitled"));
 	};
 
 	// Give the user opportunity to save temp workspaces
-	if(RequiresSave())
+	if(RequiresSave() && !bSavingWorkspaceOnly)
 	{
 		// Ensure we dont also 'save as' other externally linked assets at this point
 		TGuardValue<bool> SaveWorkspaceOnly(bSavingWorkspaceOnly, true);
@@ -312,15 +375,26 @@ bool FWorkspaceEditor::OnRequestClose(EAssetEditorCloseReason InCloseReason)
 
 void FWorkspaceEditor::OnClose()
 {
+	SaveEditedObjectState();
+	
 	if (DetailsView.IsValid())
 	{
 		DetailsView->SetObject(nullptr);
 		DetailsView.Reset();
 	}
 
-	FWorkflowCentricApplication::OnClose();
+	FBaseAssetToolkit::OnClose();
 }
 
+void FWorkspaceEditor::RegisterToolbar()
+{
+	IWorkspaceEditor::RegisterToolbar();
+}
+
+bool FWorkspaceEditor::ShouldReopenEditorForSavedAsset(const UObject* Asset) const
+{
+	return Asset && Asset->GetClass() != UWorkspace::StaticClass();
+}
 }
 
 #undef LOCTEXT_NAMESPACE
\ No newline at end of file
diff --git a/Engine/Plugins/Experimental/Workspace/Source/WorkspaceEditor/Private/WorkspaceEditor.h b/Engine/Plugins/Experimental/Workspace/Source/WorkspaceEditor/Private/WorkspaceEditor.h
--- a/Engine/Plugins/Experimental/Workspace/Source/WorkspaceEditor/Private/WorkspaceEditor.h
+++ b/Engine/Plugins/Experimental/Workspace/Source/WorkspaceEditor/Private/WorkspaceEditor.h
@@ -3,7 +3,6 @@
 #pragma once
 
 #include "CoreMinimal.h"
-#include "GraphEditor.h"
 #include "InstancedStruct.h"
 #include "IWorkspaceEditor.h"
 #include "WorkflowOrientedApp/WorkflowTabManager.h"
@@ -20,11 +19,11 @@ class IToolkitHost;
 namespace UE::Workspace
 {
 
-class FWorkspaceEditorMode;
 struct FGraphDocumentSummoner;
 struct FWorkspaceTabSummoner;
 struct FAssetDocumentSummoner;
 class FWorkspaceEditorModule;
+class SWorkspaceView;
 
 namespace WorkspaceModes
 {
@@ -37,14 +36,11 @@ namespace WorkspaceTabs
 	extern const FName WorkspaceView;
 }
 
-class FWorkspaceEditor : public IWorkspaceEditor
+class FWorkspaceEditor : public IWorkspaceEditor, public FGCObject
 {
 public:
-	FWorkspaceEditor();
-	virtual ~FWorkspaceEditor() override;
-
-	/** Edits the specified asset */
-	void InitEditor(const EToolkitMode::Type InMode, const TSharedPtr<IToolkitHost>& InInitToolkitHost, UWorkspace* InWorkspace);
+	FWorkspaceEditor(UWorkspaceAssetEditor* InOwningAssetEditor);
+	virtual ~FWorkspaceEditor() override {}
 
 private:
 	friend class FWorkspaceEditorMode;
@@ -54,21 +50,35 @@ private:
 	friend struct FAssetDocumentSummoner;
 	friend class FWorkspaceEditorModule;
 
-	// IToolkit interface
-	virtual void RegisterTabSpawners(const TSharedRef<FTabManager>& InTabManager) override;
+	// FBaseAssetToolkit interface
+	virtual void RegisterTabSpawners(const TSharedRef<class FTabManager>& TabManager) override;
 	virtual void UnregisterTabSpawners(const TSharedRef<FTabManager>& InTabManager) override;
 	virtual FName GetToolkitFName() const override;
 	virtual FText GetBaseToolkitName() const override;
 	virtual FString GetWorldCentricTabPrefix() const override;
 	virtual FLinearColor GetWorldCentricTabColorScale() const override;
+	virtual void CreateWidgets() override;
+	virtual void PostInitAssetEditor() override;
 	virtual void InitToolMenuContext(FToolMenuContext& InMenuContext) override;
 	virtual void SaveAsset_Execute() override;
 	virtual bool OnRequestClose(EAssetEditorCloseReason InCloseReason) override;
 	virtual void OnClose() override;
+	virtual void RegisterToolbar() override;
+	virtual bool ShouldReopenEditorForSavedAsset(const UObject* Asset) const override;
 
 	// FAssetEditorToolkit interface
 	virtual void GetSaveableObjects(TArray<UObject*>& OutObjects) const override;
 
+	// FGCObject interface
+	virtual void AddReferencedObjects(FReferenceCollector& Collector) override
+	{
+		Collector.AddReferencedObject(Workspace);
+	}
+	virtual FString GetReferencerName() const override
+	{
+		return TEXT("FWorkspaceEditor");
+	}
+
 	// IWorkspaceEditor interface
 	virtual void OpenAssets(TConstArrayView<FAssetData> InAssets) override;
 	virtual void OpenObjects(TConstArrayView<UObject*> InObjects) override;
@@ -82,34 +92,25 @@ private:
 
 	void ExtendToolbar();
 
-	void CloseDocumentTab(const UObject* DocumentID);
+	void CloseDocumentTab(const UObject* DocumentID) const;
 
-	void HandleDetailsViewCreated(TSharedRef<IDetailsView> InDetailsView)
-	{
-		DetailsView = InDetailsView;
-	}
-	
 	bool InEditingMode() const;
 
 	void RestoreEditedObjectState();
 
-	void SaveEditedObjectState();
+	void SaveEditedObjectState() const;
 
 	TSharedPtr<SDockTab> OpenDocument(const UObject* InForObject, FDocumentTracker::EOpenDocumentCause InCause);
 
-	void RecordDocumentState(const TInstancedStruct<FWorkspaceDocumentState>& InState);
+	void RecordDocumentState(const TInstancedStruct<FWorkspaceDocumentState>& InState) const;
 
-	// The asset we are editing
-	UWorkspace* Workspace = nullptr;
+	/** The asset being edited */
+	TObjectPtr<UWorkspace> Workspace = nullptr;
 
 	// Document tracker
 	TSharedPtr<FDocumentTracker> DocumentManager;
 
-	// Command list for this editor
-	TSharedPtr<FUICommandList> CommandList;
-
-	// Our details panel
-	TSharedPtr<IDetailsView> DetailsView;
+	TSharedPtr<SWorkspaceView> WorkspaceView;
 
 	bool bSavingWorkspaceOnly = false;
 };
diff --git a/Engine/Plugins/Experimental/Workspace/Source/WorkspaceEditor/Private/WorkspaceEditorMode.cpp b/Engine/Plugins/Experimental/Workspace/Source/WorkspaceEditor/Private/WorkspaceEditorMode.cpp
deleted file mode 100644
diff --git a/Engine/Plugins/Experimental/Workspace/Source/WorkspaceEditor/Private/WorkspaceEditorMode.h b/Engine/Plugins/Experimental/Workspace/Source/WorkspaceEditor/Private/WorkspaceEditorMode.h
deleted file mode 100644
diff --git a/Engine/Plugins/Experimental/Workspace/Source/WorkspaceEditor/Private/WorkspaceEditorModule.cpp b/Engine/Plugins/Experimental/Workspace/Source/WorkspaceEditor/Private/WorkspaceEditorModule.cpp
--- a/Engine/Plugins/Experimental/Workspace/Source/WorkspaceEditor/Private/WorkspaceEditorModule.cpp
+++ b/Engine/Plugins/Experimental/Workspace/Source/WorkspaceEditor/Private/WorkspaceEditorModule.cpp
@@ -8,6 +8,7 @@
 #include "SGraphDocument.h"
 #include "SWorkspacePicker.h"
 #include "Workspace.h"
+#include "WorkspaceAssetEditor.h"
 #include "WorkspaceEditor.h"
 #include "WorkspaceFactory.h"
 #include "AssetRegistry/AssetRegistryModule.h"
@@ -148,7 +149,7 @@ void FWorkspaceEditorModule::OpenWorkspaceForObject(UObject* InObject, EOpenWork
 		}
 	}
 
-	FWorkspaceEditor* WorkspaceEditor = nullptr;
+	UWorkspaceAssetEditor* WorkspaceEditor = nullptr;
 
 	auto HandleNewWorkspace = [InObject, &WorkspaceEditor, WorkSpaceFactoryClass]()
 	{
@@ -160,10 +161,13 @@ void FWorkspaceEditorModule::OpenWorkspaceForObject(UObject* InObject, EOpenWork
 			UWorkspace* NewWorkspace = CastChecked<UWorkspace>(Factory->FactoryCreateNew(UWorkspace::StaticClass(), Package, PackageName, RF_Public | RF_Standalone, NULL, GWarn));
 			NewWorkspace->AddAsset(InObject, false);
 			NewWorkspace->MarkPackageDirty();
-			TSharedRef<FWorkspaceEditor> Editor = MakeShared<FWorkspaceEditor>();
-			Editor->InitEditor(EToolkitMode::Standalone, nullptr, NewWorkspace);
 
-			WorkspaceEditor = &Editor.Get();
+			UAssetEditorSubsystem* AssetEditorSubsystem = GEditor->GetEditorSubsystem<UAssetEditorSubsystem>();
+			UWorkspaceAssetEditor* AssetEditor = NewObject<UWorkspaceAssetEditor>(AssetEditorSubsystem, NAME_None, RF_Transient);
+			AssetEditor->SetObjectToEdit(NewWorkspace);
+			AssetEditor->Initialize();
+
+			WorkspaceEditor = AssetEditor;
 		}
 	};
 
@@ -173,7 +177,7 @@ void FWorkspaceEditorModule::OpenWorkspaceForObject(UObject* InObject, EOpenWork
 		{
 			GEditor->GetEditorSubsystem<UAssetEditorSubsystem>()->OpenEditorForAsset(ExistingWorkspace);
 
-			WorkspaceEditor = static_cast<FWorkspaceEditor*>(GEditor->GetEditorSubsystem<UAssetEditorSubsystem>()->FindEditorForAsset(ExistingWorkspace, true));
+			WorkspaceEditor = static_cast<UWorkspaceAssetEditor*>(GEditor->GetEditorSubsystem<UAssetEditorSubsystem>()->FindEditorForAsset(ExistingWorkspace, true));
 		}
 	};
 
@@ -197,11 +201,6 @@ void FWorkspaceEditorModule::OpenWorkspaceForObject(UObject* InObject, EOpenWork
 
 		WorkspacePicker->ShowModal();
 	}
-	
-	if(WorkspaceEditor)
-	{
-		WorkspaceEditor->OpenAssets({InObject});
-	}
 }
 
 const FObjectDocumentArgs* FWorkspaceEditorModule::FindObjectDocumentType(const FTopLevelAssetPath& InClassPath) const
@@ -214,6 +213,19 @@ TArray<FTopLevelAssetPath> FWorkspaceEditorModule::GetAllowedObjectTypesForArea(
 	return DocumentAreaMap.FindRef(InSpawnLocation).Array();
 }
 
+IWorkspaceEditorModule::FOnRegisterDetailCustomizations& FWorkspaceEditorModule::OnRegisterWorkspaceDetailsCustomization()
+{
+	return OnRegisterDetailCustomizations;
+}
+
+void FWorkspaceEditorModule::ApplyWorkspaceDetailsCustomization(TSharedPtr<IDetailsView>& DetailsView) const
+{
+	if (OnRegisterDetailCustomizations.IsBound())
+	{
+		OnRegisterDetailCustomizations.Broadcast(DetailsView);
+	}
+}
+
 }
 
 IMPLEMENT_MODULE(UE::Workspace::FWorkspaceEditorModule, WorkspaceEditor);
diff --git a/Engine/Plugins/Experimental/Workspace/Source/WorkspaceEditor/Private/WorkspaceEditorModule.h b/Engine/Plugins/Experimental/Workspace/Source/WorkspaceEditor/Private/WorkspaceEditorModule.h
--- a/Engine/Plugins/Experimental/Workspace/Source/WorkspaceEditor/Private/WorkspaceEditorModule.h
+++ b/Engine/Plugins/Experimental/Workspace/Source/WorkspaceEditor/Private/WorkspaceEditorModule.h
@@ -1,4 +1,4 @@
-﻿// Copyright Epic Games, Inc. All Rights Reserved.
+// Copyright Epic Games, Inc. All Rights Reserved.
 
 #pragma once
 
@@ -24,6 +24,7 @@ private:
 	virtual void UnregisterObjectDocumentType(const FTopLevelAssetPath& InClassPath) override;
 	virtual FObjectDocumentArgs CreateGraphDocumentArgs(const FGraphDocumentWidgetArgs& InArgs) override;
 	virtual void OpenWorkspaceForObject(UObject* InObject, EOpenWorkspaceMethod InOpenMethod, const TSubclassOf<UWorkspaceFactory> WorkSpaceFactoryClass) override;
+	virtual FOnRegisterDetailCustomizations& OnRegisterWorkspaceDetailsCustomization() override;
 
 	// Find an existing registered object document type
 	const FObjectDocumentArgs* FindObjectDocumentType(const FTopLevelAssetPath& InClassPath) const;
@@ -34,10 +35,16 @@ private:
 	// Get the exported set of assets in asset registry tags for a workspace's FAssetData
 	static bool GetExportedAssetsForWorkspace(const FAssetData& InWorkspaceAsset, FWorkspaceAssetRegistryExports& OutExports);
 
+	// Applies any previously registered details-view customizations
+	void ApplyWorkspaceDetailsCustomization(TSharedPtr<IDetailsView>& DetailsView) const;
+
 	TMap<FTopLevelAssetPath, FObjectDocumentArgs> ObjectDocumentArgs;
 
 	TMap<FName, TSet<FTopLevelAssetPath>> DocumentAreaMap;
 
+	/** Event called to allow external clients to register details customizations */
+	FOnRegisterDetailCustomizations OnRegisterDetailCustomizations;
+
 	friend struct FAssetDocumentSummoner;
 	friend class FWorkspaceEditor;
 };
diff --git a/Engine/Plugins/Experimental/Workspace/Source/WorkspaceEditor/Private/WorkspaceTabSummoner.cpp b/Engine/Plugins/Experimental/Workspace/Source/WorkspaceEditor/Private/WorkspaceTabSummoner.cpp
deleted file mode 100644
diff --git a/Engine/Plugins/Experimental/Workspace/Source/WorkspaceEditor/Private/WorkspaceTabSummoner.h b/Engine/Plugins/Experimental/Workspace/Source/WorkspaceEditor/Private/WorkspaceTabSummoner.h
deleted file mode 100644
diff --git a/Engine/Plugins/Experimental/Workspace/Source/WorkspaceEditor/Public/IWorkspaceEditor.h b/Engine/Plugins/Experimental/Workspace/Source/WorkspaceEditor/Public/IWorkspaceEditor.h
--- a/Engine/Plugins/Experimental/Workspace/Source/WorkspaceEditor/Public/IWorkspaceEditor.h
+++ b/Engine/Plugins/Experimental/Workspace/Source/WorkspaceEditor/Public/IWorkspaceEditor.h
@@ -2,14 +2,18 @@
 
 #pragma once
 
-#include "WorkflowOrientedApp/WorkflowCentricApplication.h"
+#include "Tools/BaseAssetToolkit.h"
+
+class UWorkspaceAssetEditor;
 
 namespace UE::Workspace
 {
 
-class IWorkspaceEditor : public FWorkflowCentricApplication
+class IWorkspaceEditor : public FBaseAssetToolkit
 {
 public:
+	IWorkspaceEditor(UAssetEditor* InOwningAssetEditor) : FBaseAssetToolkit(InOwningAssetEditor) {}
+
 	// Open the supplied assets for editing within the workspace editor
 	virtual void OpenAssets(TConstArrayView<FAssetData> InAssets) = 0;
 
diff --git a/Engine/Plugins/Experimental/Workspace/Source/WorkspaceEditor/Public/IWorkspaceEditorModule.h b/Engine/Plugins/Experimental/Workspace/Source/WorkspaceEditor/Public/IWorkspaceEditorModule.h
--- a/Engine/Plugins/Experimental/Workspace/Source/WorkspaceEditor/Public/IWorkspaceEditorModule.h
+++ b/Engine/Plugins/Experimental/Workspace/Source/WorkspaceEditor/Public/IWorkspaceEditorModule.h
@@ -9,10 +9,12 @@
 #include "Misc/Attribute.h"
 #include "Modules/ModuleInterface.h"
 
-struct FSlateBrush;
 class SWidget;
-struct FTopLevelAssetPath;
 class UEdGraph;
+class IDetailsView;
+
+struct FSlateBrush;
+struct FTopLevelAssetPath;
 struct FWorkspaceDocumentState;
 
 namespace UE::Workspace
@@ -156,6 +158,10 @@ public:
 
 	// Make the required args for a document widget for a UEdGraph
 	virtual FObjectDocumentArgs CreateGraphDocumentArgs(const FGraphDocumentWidgetArgs& InArgs) = 0;
+
+	// Event to allow registering details customizations
+	DECLARE_EVENT_OneParam(IWorkspaceEditorModule, FOnRegisterDetailCustomizations, TSharedPtr<IDetailsView>&);
+	virtual FOnRegisterDetailCustomizations& OnRegisterWorkspaceDetailsCustomization() = 0;
 };
 
 }
