Minor workspace editor fixes
- Re-open workspace asset editor (on save) only if the editor itself is not closing down
- Dirty AssetEntry package, not the asset it is pointing to

diff --git a/Engine/Plugins/Experimental/Workspace/Source/WorkspaceEditor/Private/Workspace.cpp b/Engine/Plugins/Experimental/Workspace/Source/WorkspaceEditor/Private/Workspace.cpp
--- a/Engine/Plugins/Experimental/Workspace/Source/WorkspaceEditor/Private/Workspace.cpp
+++ b/Engine/Plugins/Experimental/Workspace/Source/WorkspaceEditor/Private/Workspace.cpp
@@ -296,7 +296,7 @@ void UWorkspace::PostLoad()
 
 		for (const TObjectPtr<UWorkspaceAssetEntry>& Entry : AssetEntries)
 		{
-			Entry->Asset->GetPackage()->SetDirtyFlag(true);	
+			Entry->GetPackage()->SetDirtyFlag(true);	
 		}
 
 		Assets_DEPRECATED.Empty();
@@ -334,4 +334,4 @@ void UWorkspace::GetAssetRegistryTags(FAssetRegistryTagsContext Context) const
 	FWorkspaceAssetRegistryExports::StaticStruct()->ExportText(TagValue, &Exports, nullptr, nullptr, PPF_None, nullptr);
 
 	Context.AddTag(FAssetRegistryTag(ExportsAssetRegistryTag, TagValue, FAssetRegistryTag::TT_Hidden));
-}
\ No newline at end of file
+}
diff --git a/Engine/Plugins/Experimental/Workspace/Source/WorkspaceEditor/Private/WorkspaceEditor.cpp b/Engine/Plugins/Experimental/Workspace/Source/WorkspaceEditor/Private/WorkspaceEditor.cpp
--- a/Engine/Plugins/Experimental/Workspace/Source/WorkspaceEditor/Private/WorkspaceEditor.cpp
+++ b/Engine/Plugins/Experimental/Workspace/Source/WorkspaceEditor/Private/WorkspaceEditor.cpp
@@ -355,6 +355,8 @@ void FWorkspaceEditor::RecordDocumentState(const TInstancedStruct<FWorkspaceDocu
 
 bool FWorkspaceEditor::OnRequestClose(EAssetEditorCloseReason InCloseReason)
 {
+	TGuardValue<bool> ClosingDown(bClosingDown, true);
+
 	auto RequiresSave = [this]()
 	{
 		const UPackage* Package = Workspace->GetOutermost();
@@ -393,7 +395,7 @@ void FWorkspaceEditor::RegisterToolbar()
 
 bool FWorkspaceEditor::ShouldReopenEditorForSavedAsset(const UObject* Asset) const
 {
-	return Asset && Asset->GetClass() != UWorkspace::StaticClass();
+	return !bClosingDown;
 }
 }
 
diff --git a/Engine/Plugins/Experimental/Workspace/Source/WorkspaceEditor/Private/WorkspaceEditor.h b/Engine/Plugins/Experimental/Workspace/Source/WorkspaceEditor/Private/WorkspaceEditor.h
--- a/Engine/Plugins/Experimental/Workspace/Source/WorkspaceEditor/Private/WorkspaceEditor.h
+++ b/Engine/Plugins/Experimental/Workspace/Source/WorkspaceEditor/Private/WorkspaceEditor.h
@@ -113,6 +113,7 @@ private:
 	TSharedPtr<SWorkspaceView> WorkspaceView;
 
 	bool bSavingWorkspaceOnly = false;
+	bool bClosingDown = false;
 };
 
 }
\ No newline at end of file
