UE.TargetAutomation(RunTest=Animation) - Linux and Mac_arm64 - StructProperty and FloatProperty is not initialized properly
#fix ensure properties have default values

diff --git a/Engine/Plugins/Experimental/Workspace/Source/WorkspaceEditor/Private/GraphDocumentState.h b/Engine/Plugins/Experimental/Workspace/Source/WorkspaceEditor/Private/GraphDocumentState.h
--- a/Engine/Plugins/Experimental/Workspace/Source/WorkspaceEditor/Private/GraphDocumentState.h
+++ b/Engine/Plugins/Experimental/Workspace/Source/WorkspaceEditor/Private/GraphDocumentState.h
@@ -20,8 +20,8 @@ struct FGraphDocumentState : public FWorkspaceDocumentState
 	{}
 
 	UPROPERTY()
-	FVector2D ViewLocation;
+	FVector2D ViewLocation = FVector2D::Zero();
 
 	UPROPERTY()
-	float ZoomAmount;
+	float ZoomAmount = 0.f;
 };
