WorkspaceEditor: Add navigation delegates to SGraphDocument

diff --git a/Engine/Plugins/Experimental/Workspace/Source/WorkspaceEditor/Private/SGraphDocument.cpp b/Engine/Plugins/Experimental/Workspace/Source/WorkspaceEditor/Private/SGraphDocument.cpp
--- a/Engine/Plugins/Experimental/Workspace/Source/WorkspaceEditor/Private/SGraphDocument.cpp
+++ b/Engine/Plugins/Experimental/Workspace/Source/WorkspaceEditor/Private/SGraphDocument.cpp
@@ -54,6 +54,8 @@ void SGraphDocument::Construct(const FArguments& InArgs, TSharedRef<FWorkspaceEd
 		.GraphToEdit(InGraph)
 		.GraphEvents(Events)
 		.AssetEditorToolkit(HostingAppPtr)
+		.OnNavigateHistoryBack(InArgs._OnNavigateHistoryBack)		
+		.OnNavigateHistoryForward(InArgs._OnNavigateHistoryForward)
 	];
 }
 
diff --git a/Engine/Plugins/Experimental/Workspace/Source/WorkspaceEditor/Private/SGraphDocument.h b/Engine/Plugins/Experimental/Workspace/Source/WorkspaceEditor/Private/SGraphDocument.h
--- a/Engine/Plugins/Experimental/Workspace/Source/WorkspaceEditor/Private/SGraphDocument.h
+++ b/Engine/Plugins/Experimental/Workspace/Source/WorkspaceEditor/Private/SGraphDocument.h
@@ -41,6 +41,9 @@ class SGraphDocument : public SCompoundWidget
 	SLATE_EVENT(FOnCanPerformActionOnSelectedNodes, OnCanDuplicateSelectedNodes)
 	SLATE_EVENT(FOnDuplicateSelectedNodes, OnDuplicateSelectedNodes)
 
+	SLATE_EVENT(FSimpleDelegate, OnNavigateHistoryBack)
+	SLATE_EVENT(FSimpleDelegate, OnNavigateHistoryForward)
+
 	SLATE_END_ARGS()
 
 	void Construct(const FArguments& InArgs, TSharedRef<FWorkspaceEditor> InHostingApp, UEdGraph* InGraph);
