Workspace outliner initial implementation:

- Added Outliner implementation (based off SSceneOutliner)
- Added OutlinerItem asset registry export definitions
- Added concept of "Details" provider to improve/extend asset-registry exported Outliner items
- Added RigVMAssetUserData to export AnimNext graph entries to Workspace outliner

- WorkspaceEditor document navigation is now brower-like, with reusing of tabs and hotkey/button navigation for backwards and forwards
- WorkspaceEditor document tabs are now wrapped withing SWorkspaceTabWrapper, providing a fixed breadcrumb bar
- All AnimNext documents are now spawned in the MiddleDocumentArea
- AssetDocument summoner no longer _by default_  allows spawning of tabs when SupportedAssetClasses is empty

diff --git a/Engine/Plugins/Experimental/AnimNext/Source/AnimNext/Private/AnimNextRigVMAsset.cpp b/Engine/Plugins/Experimental/AnimNext/Source/AnimNext/Private/AnimNextRigVMAsset.cpp
--- a/Engine/Plugins/Experimental/AnimNext/Source/AnimNext/Private/AnimNextRigVMAsset.cpp
+++ b/Engine/Plugins/Experimental/AnimNext/Source/AnimNext/Private/AnimNextRigVMAsset.cpp
@@ -51,4 +51,15 @@ void UAnimNextRigVMAsset::GetAssetRegistryTags(FAssetRegistryTagsContext Context
 		EditorData->GetAssetRegistryTags(Context);
 	}
 #endif
+
+#if WITH_EDITOR
+	// Allow asset user data to output tags
+	for(const UAssetUserData* AssetUserDataItem : *GetAssetUserDataArray())
+	{
+		if (AssetUserDataItem)
+		{
+			AssetUserDataItem->GetAssetRegistryTags(Context);
+		}
+	}
+#endif // WITH_EDITOR
 }
diff --git a/Engine/Plugins/Experimental/AnimNext/Source/AnimNext/Private/Graph/AnimNextGraph.cpp b/Engine/Plugins/Experimental/AnimNext/Source/AnimNext/Private/Graph/AnimNextGraph.cpp
--- a/Engine/Plugins/Experimental/AnimNext/Source/AnimNext/Private/Graph/AnimNextGraph.cpp
+++ b/Engine/Plugins/Experimental/AnimNext/Source/AnimNext/Private/Graph/AnimNextGraph.cpp
@@ -22,6 +22,11 @@
 #include "UObject/LinkerLoad.h"
 #include "UObject/ObjectResource.h"
 
+#if WITH_EDITOR	
+#include "Engine/ExternalAssetDependencyGatherer.h"
+REGISTER_ASSETDEPENDENCY_GATHERER(FExternalAssetDependencyGatherer, UAnimNextGraph);
+#endif // WITH_EDITOR
+
 DEFINE_STAT(STAT_AnimNext_Graph_AllocateInstance);
 DEFINE_STAT(STAT_AnimNext_Graph_UpdateParamLayer);
 
diff --git a/Engine/Plugins/Experimental/AnimNext/Source/AnimNextEditor/Private/AnimNextEditorModule.cpp b/Engine/Plugins/Experimental/AnimNext/Source/AnimNextEditor/Private/AnimNextEditorModule.cpp
--- a/Engine/Plugins/Experimental/AnimNext/Source/AnimNextEditor/Private/AnimNextEditorModule.cpp
+++ b/Engine/Plugins/Experimental/AnimNext/Source/AnimNextEditor/Private/AnimNextEditorModule.cpp
@@ -38,6 +38,8 @@
 #include "Param/ObjectFunctionLocatorEditor.h"
 #include "Param/ObjectPropertyLocatorEditor.h"
 #include "Param/SAddParametersDialog.h"
+#include "AnimNextRigVMWorkspaceAssetUserData.h"
+#include "Graph/AnimNextGraph_OutlinerItemDetails.h"
 
 #define LOCTEXT_NAMESPACE "AnimNextEditorModule"
 
@@ -182,6 +184,11 @@ class FModule : public IModule
 		UolEditorModule.RegisterLocatorEditor("AnimNextObjectCast", MakeShared<FObjectCastLocatorEditor>());
 
 		UolEditorModule.RegisterEditorContext("AnimNextContext", MakeShared<FAnimNextLocatorContext>());
+
+		Workspace::IWorkspaceEditorModule& WorkspaceModule = FModuleManager::Get().LoadModuleChecked<Workspace::IWorkspaceEditorModule>("WorkspaceEditor");
+		const TSharedPtr<FAnimNextGraphItemDetails> GraphItemDetails = MakeShareable<FAnimNextGraphItemDetails>(new FAnimNextGraphItemDetails());
+		WorkspaceModule.RegisterWorkspaceItemDetails(Workspace::FOutlinerItemDetailsId(FAnimNextGraphOutlinerData::StaticStruct()->GetFName()), StaticCastSharedPtr<UE::Workspace::IWorkspaceOutlinerItemDetails>(GraphItemDetails));
+		FAnimNextGraphItemDetails::RegisterToolMenuExtensions();
 	}
 
 	virtual void ShutdownModule() override
@@ -216,6 +223,10 @@ class FModule : public IModule
 
 			UolEditorModule.UnregisterEditorContext("AnimNextContext");
 		}
+
+		Workspace::IWorkspaceEditorModule& WorkspaceModule = FModuleManager::Get().LoadModuleChecked<Workspace::IWorkspaceEditorModule>("WorkspaceEditor");
+		WorkspaceModule.UnregisterWorkspaceItemDetails(Workspace::FOutlinerItemDetailsId(FAnimNextGraphOutlinerData::StaticStruct()->GetFName()));
+		FAnimNextGraphItemDetails::UnregisterToolMenuExtensions();
 	}
 
 	virtual TSharedRef<SWidget> CreateParameterPicker(const FParameterPickerArgs& InArgs) override
@@ -227,23 +238,42 @@ class FModule : public IModule
 	void RegisterWorkspaceDocumentTypes(Workspace::IWorkspaceEditorModule& WorkspaceEditorModule)
 	{
 		// --- AnimNextSchedule ---
-		WorkspaceEditorModule.RegisterObjectDocumentType(FTopLevelAssetPath(TEXT("/Script/AnimNext.AnimNextSchedule")),
-			Workspace::FObjectDocumentArgs(
-				Workspace::FOnMakeDocumentWidget::CreateLambda([](const Workspace::FWorkspaceEditorContext& InContext)
-				{
-					UAnimNextSchedule* Schedule = CastChecked<UAnimNextSchedule>(InContext.Object);
-					FPropertyEditorModule& PropertyEditorModule = FModuleManager::GetModuleChecked<FPropertyEditorModule>( "PropertyEditor" );
-					FDetailsViewArgs DetailsViewArgs;
-					DetailsViewArgs.NameAreaSettings = FDetailsViewArgs::HideNameArea;
-					TSharedRef<IDetailsView> DetailsView = PropertyEditorModule.CreateDetailView(DetailsViewArgs);
-					DetailsView->SetObject(Schedule);
-					return DetailsView;
-				}),
-				Workspace::WorkspaceTabs::MiddleDocumentArea));
+		Workspace::FObjectDocumentArgs ScheduleDocumentArgs(
+			Workspace::FOnMakeDocumentWidget::CreateLambda([](const Workspace::FWorkspaceEditorContext& InContext)
+			{
+				UAnimNextSchedule* Schedule = CastChecked<UAnimNextSchedule>(InContext.Object);
+				FPropertyEditorModule& PropertyEditorModule = FModuleManager::GetModuleChecked<FPropertyEditorModule>( "PropertyEditor" );
+				FDetailsViewArgs DetailsViewArgs;
+				DetailsViewArgs.NameAreaSettings = FDetailsViewArgs::HideNameArea;
+				TSharedRef<IDetailsView> DetailsView = PropertyEditorModule.CreateDetailView(DetailsViewArgs);
+				DetailsView->SetObject(Schedule);
+				return DetailsView;
+			}
+		), Workspace::WorkspaceTabs::MiddleDocumentArea);
+		ScheduleDocumentArgs.OnGetTabName = Workspace::FOnGetTabName::CreateLambda([](const Workspace::FWorkspaceEditorContext& InContext)
+		{
+			const UAnimNextSchedule* Schedule = CastChecked<UAnimNextSchedule>(InContext.Object);
+			return FText::FromName(Schedule->GetFName());
+		});
+
+		ScheduleDocumentArgs.OnGetDocumentBreadcrumbTrail = Workspace::FOnGetDocumentBreadcrumbTrail::CreateLambda([](const Workspace::FWorkspaceEditorContext& InContext, TArray<TSharedPtr<Workspace::FWorkspaceBreadcrumb>>& OutBreadcrumbs)
+		{
+			if (const UAnimNextSchedule* Schedule = Cast<UAnimNextSchedule>(InContext.Object))
+			{
+				const TSharedPtr<Workspace::FWorkspaceBreadcrumb>& GraphCrumb = OutBreadcrumbs.Add_GetRef(MakeShared<Workspace::FWorkspaceBreadcrumb>());
+				GraphCrumb->OnGetLabel = Workspace::FWorkspaceBreadcrumb::FOnGetBreadcrumbLabel::CreateLambda([ScheduleName = Schedule->GetFName()]{ return FText::FromName(ScheduleName); });
+			}
+		});
+
+		ScheduleDocumentArgs.OnGetTabIcon = Workspace::FOnGetTabIcon::CreateLambda([](const Workspace::FWorkspaceEditorContext& InContext)
+		{
+			return FAppStyle::GetBrush(TEXT("ClassIcon.Default"));
+		});
+
+		WorkspaceEditorModule.RegisterObjectDocumentType(FTopLevelAssetPath(TEXT("/Script/AnimNext.AnimNextSchedule")), ScheduleDocumentArgs);
 
 		// --- AnimNextGraph ---
-		WorkspaceEditorModule.RegisterObjectDocumentType(FTopLevelAssetPath(TEXT("/Script/AnimNext.AnimNextGraph")),
-			Workspace::FObjectDocumentArgs(
+		Workspace::FObjectDocumentArgs AnimNextGraphDocumentArgs(
 				Workspace::FOnMakeDocumentWidget::CreateLambda([](const Workspace::FWorkspaceEditorContext& InContext)
 				{
 					UAnimNextGraph* Graph = CastChecked<UAnimNextGraph>(InContext.Object);
@@ -308,7 +338,28 @@ class FModule : public IModule
 							}
 						});
 				}),
-				Workspace::WorkspaceTabs::LeftDocumentArea));
+				Workspace::WorkspaceTabs::MiddleDocumentArea);
+		AnimNextGraphDocumentArgs.OnGetTabName = Workspace::FOnGetTabName::CreateLambda([](const Workspace::FWorkspaceEditorContext& InContext)
+		{
+			const UAnimNextGraph* Graph = CastChecked<UAnimNextGraph>(InContext.Object);
+			return FText::FromName(Graph->GetFName());
+		});
+
+		AnimNextGraphDocumentArgs.OnGetDocumentBreadcrumbTrail = Workspace::FOnGetDocumentBreadcrumbTrail::CreateLambda([](const Workspace::FWorkspaceEditorContext& InContext, TArray<TSharedPtr<Workspace::FWorkspaceBreadcrumb>>& OutBreadcrumbs)
+		{
+			if (const UAnimNextGraph* Graph = Cast<UAnimNextGraph>(InContext.Object))
+			{
+				const TSharedPtr<Workspace::FWorkspaceBreadcrumb>& GraphCrumb = OutBreadcrumbs.Add_GetRef(MakeShared<Workspace::FWorkspaceBreadcrumb>());
+				GraphCrumb->OnGetLabel = Workspace::FWorkspaceBreadcrumb::FOnGetBreadcrumbLabel::CreateLambda([GraphName = Graph->GetFName()]{ return FText::FromName(GraphName); });
+			}
+		});
+
+		AnimNextGraphDocumentArgs.OnGetTabIcon = Workspace::FOnGetTabIcon::CreateLambda([](const Workspace::FWorkspaceEditorContext& InContext)
+		{
+			return FAppStyle::GetBrush(TEXT("ClassIcon.Default"));
+		});
+
+		WorkspaceEditorModule.RegisterObjectDocumentType(FTopLevelAssetPath(TEXT("/Script/AnimNext.AnimNextGraph")), AnimNextGraphDocumentArgs);
 
 		// --- AnimNextGraph_EdGraph ---
 		Workspace::FGraphDocumentWidgetArgs GraphArgs;
@@ -618,7 +669,40 @@ class FModule : public IModule
 			InContext.WorkspaceEditor->SetDetailsObjects(NewSelection.Array());
 		});
 
-		WorkspaceEditorModule.RegisterObjectDocumentType(FTopLevelAssetPath(TEXT("/Script/AnimNextUncookedOnly.AnimNextGraph_EdGraph")), WorkspaceEditorModule.CreateGraphDocumentArgs(GraphArgs));
+		Workspace::FObjectDocumentArgs GraphDocumentArgs = WorkspaceEditorModule.CreateGraphDocumentArgs(GraphArgs);
+		GraphDocumentArgs.OnGetDocumentBreadcrumbTrail = Workspace::FOnGetDocumentBreadcrumbTrail::CreateLambda([](const Workspace::FWorkspaceEditorContext& InContext, TArray<TSharedPtr<Workspace::FWorkspaceBreadcrumb>>& OutBreadcrumbs)
+		{
+			if (const URigVMEdGraph* RigVMEdGraph = Cast<URigVMEdGraph>(InContext.Object))
+			{
+				const TSharedPtr<Workspace::FWorkspaceBreadcrumb>& GraphCrumb = OutBreadcrumbs.Add_GetRef(MakeShared<Workspace::FWorkspaceBreadcrumb>());
+				GraphCrumb->OnGetLabel = Workspace::FWorkspaceBreadcrumb::FOnGetBreadcrumbLabel::CreateLambda([GraphName = RigVMEdGraph->GetFName()]{ return FText::FromName(GraphName); });
+
+				if (const UAnimNextRigVMAssetEntry* OuterAssetEntry = CastChecked<UAnimNextRigVMAssetEntry>(RigVMEdGraph->GetOuter()))
+				{
+					if (UAnimNextGraph_EditorData* OuterEditorData = CastChecked<UAnimNextGraph_EditorData>(OuterAssetEntry->GetOuter()))
+					{
+						if(UAnimNextGraph* OuterGraph = UncookedOnly::FUtils::GetGraph(OuterEditorData))
+						{
+							const TSharedPtr<Workspace::FWorkspaceBreadcrumb>& OuterGraphCrumb = OutBreadcrumbs.Add_GetRef(MakeShared<Workspace::FWorkspaceBreadcrumb>());
+							TWeakObjectPtr<UAnimNextGraph> WeakOuterGraph = OuterGraph;
+							TWeakPtr<Workspace::IWorkspaceEditor> WeakWorkspaceEditor = InContext.WorkspaceEditor;
+							OuterGraphCrumb->OnGetLabel = Workspace::FWorkspaceBreadcrumb::FOnGetBreadcrumbLabel::CreateLambda([GraphName = OuterGraph->GetFName()]{ return FText::FromName(GraphName); });
+							OuterGraphCrumb->OnClicked = Workspace::FWorkspaceBreadcrumb::FOnBreadcrumbClicked::CreateLambda(
+								[WeakOuterGraph, WeakWorkspaceEditor]
+								{
+									if (const TSharedPtr<Workspace::IWorkspaceEditor> SharedWorkspaceEditor = WeakWorkspaceEditor.Pin())
+									{
+										SharedWorkspaceEditor->OpenObjects({WeakOuterGraph.Get()});
+									}
+								}
+							);
+						}
+					}
+				}
+			}
+		});
+
+		WorkspaceEditorModule.RegisterObjectDocumentType(FTopLevelAssetPath(TEXT("/Script/AnimNextUncookedOnly.AnimNextGraph_EdGraph")), GraphDocumentArgs);
 	}
 
 	void UnregisterWorkspaceDocumentTypes()
diff --git a/Engine/Plugins/Experimental/AnimNext/Source/AnimNextEditor/Private/Graph/AnimNextGraph_OutlinerItemDetails.cpp b/Engine/Plugins/Experimental/AnimNext/Source/AnimNextEditor/Private/Graph/AnimNextGraph_OutlinerItemDetails.cpp
new file mode 100644
--- /dev/null
+++ b/Engine/Plugins/Experimental/AnimNext/Source/AnimNextEditor/Private/Graph/AnimNextGraph_OutlinerItemDetails.cpp
@@ -0,0 +1,125 @@
+﻿// Copyright Epic Games, Inc. All Rights Reserved.
+
+#include "AnimNextGraph_OutlinerItemDetails.h"
+
+#include "Toolkits/AssetEditorToolkitMenuContext.h"
+#include "Graph/AnimNextGraph_AnimationGraph.h"
+#include "Graph/AnimNextGraph_EventGraph.h"
+#include "InstancedStruct.h"
+#include "AnimNextRigVMWorkspaceAssetUserData.h"
+#include "RigVMModel/RigVMGraph.h"
+#include "WorkspaceItemMenuContext.h"
+#include "IWorkspaceEditor.h"
+#include "RigVMModel/RigVMClient.h"
+
+#define LOCTEXT_NAMESPACE "FAnimNextGraphItemDetails"
+
+namespace UE::AnimNext::Editor
+{
+void FAnimNextGraphItemDetails::HandleDoubleClick(const FToolMenuContext& ToolMenuContext) const
+{
+	const UWorkspaceItemMenuContext* WorkspaceItemContext = ToolMenuContext.FindContext<UWorkspaceItemMenuContext>();
+	const UAssetEditorToolkitMenuContext* AssetEditorContext = ToolMenuContext.FindContext<UAssetEditorToolkitMenuContext>(); 
+	if (WorkspaceItemContext && AssetEditorContext)
+	{
+		if(const TSharedPtr<UE::Workspace::IWorkspaceEditor> WorkspaceEditor = StaticCastSharedPtr<UE::Workspace::IWorkspaceEditor>(AssetEditorContext->Toolkit.Pin()))
+		{
+			const TInstancedStruct<FWorkspaceOutlinerItemData>& Data = WorkspaceItemContext->SelectedExports[0].Data;
+			if (Data.IsValid() && Data.GetScriptStruct() == FAnimNextGraphOutlinerData::StaticStruct())
+			{
+				const FAnimNextGraphOutlinerData& GraphData = Data.Get<FAnimNextGraphOutlinerData>();							
+				if (GraphData.GraphInterface)
+				{
+					if (URigVMGraph* RigVMGraph = GraphData.GraphInterface->GetRigVMGraph())
+					{
+						if(const IRigVMClientHost* RigVMClientHost = RigVMGraph->GetImplementingOuter<IRigVMClientHost>())
+						{
+							if(UObject* EditorObject = RigVMClientHost->GetEditorObjectForRigVMGraph(RigVMGraph))
+							{
+								WorkspaceEditor->OpenObjects({EditorObject});
+							}
+						}
+					}
+				}
+			}
+		}
+	}
+}
+
+UPackage* FAnimNextGraphItemDetails::GetPackage(const FWorkspaceOutlinerItemExport& Export) const 
+{
+	const TInstancedStruct<FWorkspaceOutlinerItemData>& Data = Export.Data;
+	if (Data.IsValid() && Data.GetScriptStruct() == FAnimNextGraphOutlinerData::StaticStruct())
+	{
+		const FAnimNextGraphOutlinerData& GraphData = Data.Get<FAnimNextGraphOutlinerData>();
+		if (GraphData.GraphInterface)
+		{
+			return GraphData.GraphInterface.GetObject()->GetExternalPackage();
+		}
+	}
+	return nullptr;
+}
+
+const FSlateBrush* FAnimNextGraphItemDetails::GetItemIcon() const
+{
+	return FAppStyle::GetBrush(TEXT("GraphEditor.EventGraph_24x"));
+}
+
+void FAnimNextGraphItemDetails::RegisterToolMenuExtensions()
+{
+	FToolMenuOwnerScoped OwnerScoped(TEXT("FAnimNextGraphItemDetails"));
+	if (UToolMenu* Menu = UToolMenus::Get()->ExtendMenu("WorkspaceOutliner.ItemContextMenu"))
+	{
+		Menu->AddDynamicSection(TEXT("AnimNextGraphItem"), FNewToolMenuDelegate::CreateLambda([](UToolMenu* InMenu)
+		{
+			UWorkspaceItemMenuContext* WorkspaceItemContext = InMenu->FindContext<UWorkspaceItemMenuContext>();
+			const UAssetEditorToolkitMenuContext* AssetEditorContext = InMenu->FindContext<UAssetEditorToolkitMenuContext>();			
+			if (WorkspaceItemContext && AssetEditorContext)
+			{
+				FToolMenuSection& Section = InMenu->AddSection("WorkspaceOutliner.ItemContextMenu.RootAsset", FText::FromString(TEXT("Animation Next")));
+				if(const TSharedPtr<UE::Workspace::IWorkspaceEditor> WorkspaceEditor = StaticCastSharedPtr<UE::Workspace::IWorkspaceEditor>(AssetEditorContext->Toolkit.Pin()))
+				{
+					if (WorkspaceItemContext->SelectedExports.Num() == 1)
+					{
+						TInstancedStruct<FWorkspaceOutlinerItemData>& Data = WorkspaceItemContext->SelectedExports[0].Data;
+						if (Data.IsValid() && Data.GetScriptStruct() == FAnimNextGraphOutlinerData::StaticStruct())
+						{							
+							Section.AddMenuEntry(
+							TEXT("Open"),
+							LOCTEXT("GoHereActionLabel", "Open Graph"),
+							LOCTEXT("GoHereActionTooltip", "Open selected AnimNext Graph"),
+							FSlateIcon(FAppStyle::GetAppStyleSetName(), "GraphEditor.EventGraph_16x"),
+							FUIAction(
+								FExecuteAction::CreateWeakLambda(WorkspaceItemContext, [&Data, WorkspaceEditor]()
+								{										
+									const FAnimNextGraphOutlinerData& GraphData = Data.Get<FAnimNextGraphOutlinerData>();
+									if (GraphData.GraphInterface)
+									{
+										if (URigVMGraph* RigVMGraph = GraphData.GraphInterface->GetRigVMGraph())
+										{
+											if(const IRigVMClientHost* RigVMClientHost = RigVMGraph->GetImplementingOuter<IRigVMClientHost>())
+											{
+												if(UObject* EditorObject = RigVMClientHost->GetEditorObjectForRigVMGraph(RigVMGraph))
+												{
+													WorkspaceEditor->OpenObjects({EditorObject});
+												}
+											}
+										}
+									}
+								})
+							));
+						}
+					}
+				}
+			}
+		}));
+	}
+}
+
+void FAnimNextGraphItemDetails::UnregisterToolMenuExtensions()
+{
+	UToolMenus::Get()->UnregisterOwnerByName("FAnimNextGraphItemDetails");
+}
+}
+
+#undef LOCTEXT_NAMESPACE // "FAnimNextGraphItemDetails"
\ No newline at end of file
diff --git a/Engine/Plugins/Experimental/AnimNext/Source/AnimNextEditor/Private/Graph/AnimNextGraph_OutlinerItemDetails.h b/Engine/Plugins/Experimental/AnimNext/Source/AnimNextEditor/Private/Graph/AnimNextGraph_OutlinerItemDetails.h
new file mode 100644
--- /dev/null
+++ b/Engine/Plugins/Experimental/AnimNext/Source/AnimNextEditor/Private/Graph/AnimNextGraph_OutlinerItemDetails.h
@@ -0,0 +1,21 @@
+﻿// Copyright Epic Games, Inc. All Rights Reserved.
+
+#pragma once
+
+#include "IWorkspaceOutlinerItemDetails.h"
+
+namespace UE::AnimNext::Editor
+{
+class FAnimNextGraphItemDetails : public UE::Workspace::IWorkspaceOutlinerItemDetails
+{
+public:
+	FAnimNextGraphItemDetails() = default;
+	virtual ~FAnimNextGraphItemDetails() override = default;
+	virtual void HandleDoubleClick(const FToolMenuContext& ToolMenuContext) const override;	
+	virtual UPackage* GetPackage(const FWorkspaceOutlinerItemExport& Export) const override;
+	virtual const FSlateBrush* GetItemIcon() const override;
+
+	static void RegisterToolMenuExtensions();
+	static void UnregisterToolMenuExtensions();
+};
+}
diff --git a/Engine/Plugins/Experimental/AnimNext/Source/AnimNextUncookedOnly/AnimNextUncookedOnly.Build.cs b/Engine/Plugins/Experimental/AnimNext/Source/AnimNextUncookedOnly/AnimNextUncookedOnly.Build.cs
--- a/Engine/Plugins/Experimental/AnimNext/Source/AnimNextUncookedOnly/AnimNextUncookedOnly.Build.cs
+++ b/Engine/Plugins/Experimental/AnimNext/Source/AnimNextUncookedOnly/AnimNextUncookedOnly.Build.cs
@@ -34,6 +34,7 @@ namespace UnrealBuildTool.Rules
 					"ToolMenus",
 					"UniversalObjectLocator", 
 					"UniversalObjectLocatorEditor",
+					"WorkspaceEditor"
 				}
 			);
 
diff --git a/Engine/Plugins/Experimental/AnimNext/Source/AnimNextUncookedOnly/Internal/AnimNextRigVMAssetEditorData.h b/Engine/Plugins/Experimental/AnimNext/Source/AnimNextUncookedOnly/Internal/AnimNextRigVMAssetEditorData.h
--- a/Engine/Plugins/Experimental/AnimNext/Source/AnimNextUncookedOnly/Internal/AnimNextRigVMAssetEditorData.h
+++ b/Engine/Plugins/Experimental/AnimNext/Source/AnimNextUncookedOnly/Internal/AnimNextRigVMAssetEditorData.h
@@ -87,6 +87,8 @@ protected:
 	virtual void PostDuplicate(EDuplicateMode::Type DuplicateMode) override;
 	virtual void GetAssetRegistryTags(FAssetRegistryTagsContext Context) const override;
 	virtual bool IsEditorOnly() const override { return true; }
+	virtual bool Rename(const TCHAR* NewName = nullptr, UObject* NewOuter = nullptr, ERenameFlags Flags = REN_None) override;	
+	virtual void PreDuplicate(FObjectDuplicationParameters& DupParams) override;
 
 	void HandlePackageDone(const FEndLoadPackageContext& Context);
 	void HandlePackageDone();
diff --git a/Engine/Plugins/Experimental/AnimNext/Source/AnimNextUncookedOnly/Internal/AnimNextRigVMWorkspaceAssetUserData.h b/Engine/Plugins/Experimental/AnimNext/Source/AnimNextUncookedOnly/Internal/AnimNextRigVMWorkspaceAssetUserData.h
new file mode 100644
--- /dev/null
+++ b/Engine/Plugins/Experimental/AnimNext/Source/AnimNextUncookedOnly/Internal/AnimNextRigVMWorkspaceAssetUserData.h
@@ -0,0 +1,55 @@
+// Copyright Epic Games, Inc. All Rights Reserved.
+
+#pragma once
+
+#include "WorkspaceAssetRegistryInfo.h"
+#include "Param/ParamType.h"
+#include "IAnimNextRigVMGraphInterface.h"
+
+#include "AnimNextRigVMWorkspaceAssetUserData.generated.h"
+
+USTRUCT()
+struct FAnimNextGraphAssetOutlinerData : public FWorkspaceOutlinerItemData
+{
+	GENERATED_BODY()
+	
+	FAnimNextGraphAssetOutlinerData() = default;
+};
+
+USTRUCT()
+struct FAnimNextParameterOutlinerData : public FWorkspaceOutlinerItemData
+{
+	GENERATED_BODY()
+	
+	FAnimNextParameterOutlinerData() = default;
+
+	UPROPERTY(VisibleAnywhere, Category=AnimNext)
+	FAnimNextParamType Type;
+};
+
+USTRUCT()
+struct FAnimNextGraphOutlinerData : public FWorkspaceOutlinerItemData
+{
+	GENERATED_BODY()
+	
+	FAnimNextGraphOutlinerData() = default;
+
+	UPROPERTY(VisibleAnywhere, Category=AnimNext)
+	TScriptInterface<IAnimNextRigVMGraphInterface> GraphInterface;
+};
+
+USTRUCT()
+struct FAnimNextSchedulerData : public FAnimNextGraphAssetOutlinerData
+{
+	GENERATED_BODY()
+	
+	FAnimNextSchedulerData() = default;
+};
+
+UCLASS()
+class UAnimNextGraphWorkspaceAssetUserData : public UAssetUserData
+{	
+	GENERATED_BODY()
+
+	virtual void GetAssetRegistryTags(FAssetRegistryTagsContext Context) const override;
+};
\ No newline at end of file
diff --git a/Engine/Plugins/Experimental/AnimNext/Source/AnimNextUncookedOnly/Internal/UncookedOnlyUtils.h b/Engine/Plugins/Experimental/AnimNext/Source/AnimNextUncookedOnly/Internal/UncookedOnlyUtils.h
--- a/Engine/Plugins/Experimental/AnimNext/Source/AnimNextUncookedOnly/Internal/UncookedOnlyUtils.h
+++ b/Engine/Plugins/Experimental/AnimNext/Source/AnimNextUncookedOnly/Internal/UncookedOnlyUtils.h
@@ -13,6 +13,8 @@
 
 struct FAnimNextEditorParam;
 struct FAnimNextParam;
+struct FEdGraphPinType;
+struct FWorkspaceOutlinerItemExports;
 class UAnimNextSchedule;
 class UAnimNextGraph;
 class UAnimNextGraph_EditorData;
@@ -20,7 +22,6 @@ class UAnimNextGraph_EdGraph;
 class URigVMController;
 class URigVMGraph;
 class UAnimNextGraph_EdGraph;
-struct FEdGraphPinType;
 class UAnimNextRigVMAsset;
 class UAnimNextRigVMAssetEditorData;
 class UAnimNextRigVMAssetEntry;
@@ -29,7 +30,7 @@ namespace UE
 {
 	namespace AnimNext
 	{
-		static const FName ExportsAnimNextAssetRegistryTag = TEXT("AnimNextExports");
+		static const FLazyName ExportsAnimNextAssetRegistryTag = TEXT("AnimNextExports");
 	}
 }
 
@@ -188,6 +189,9 @@ struct ANIMNEXTUNCOOKEDONLY_API FUtils
 	static void GetBlueprintParameters(const UBlueprint* InBlueprint, FAnimNextParameterProviderAssetRegistryExports& OutExports);
 	static void GetBlueprintParameters(const UBlueprint* InBlueprint, TSet<FAnimNextParameterAssetRegistryExportEntry>& OutExports);
 
+	// Gets the asset-registry information needed for representing the contained data into the Workspace Outliner
+	static void GetAssetOutlinerItems(const UAnimNextRigVMAssetEditorData* EditorData, FWorkspaceOutlinerItemExports& OutExports);
+
 	// Attempts to determine the type from a parameter name
 	// If the name cannot be found, the returned type will be invalid
 	// Note that this is expensive and can query the asset registry
diff --git a/Engine/Plugins/Experimental/AnimNext/Source/AnimNextUncookedOnly/Private/AnimNextRigVMAssetEditorData.cpp b/Engine/Plugins/Experimental/AnimNext/Source/AnimNextUncookedOnly/Private/AnimNextRigVMAssetEditorData.cpp
--- a/Engine/Plugins/Experimental/AnimNext/Source/AnimNextUncookedOnly/Private/AnimNextRigVMAssetEditorData.cpp
+++ b/Engine/Plugins/Experimental/AnimNext/Source/AnimNextUncookedOnly/Private/AnimNextRigVMAssetEditorData.cpp
@@ -5,6 +5,7 @@
 #include "AnimNextRigVMAsset.h"
 #include "AnimNextRigVMAssetEntry.h"
 #include "AnimNextRigVMAssetSchema.h"
+#include "AnimNextRigVMWorkspaceAssetUserData.h"
 #include "ControlRigDefines.h"
 #include "ExternalPackageHelper.h"
 #include "IAnimNextRigVMGraphInterface.h"
@@ -69,6 +70,12 @@ void UAnimNextRigVMAssetEditorData::Serialize(FArchive& Ar)
 	RigVMClient.SetDefaultSchemaClass(UAnimNextRigVMAssetSchema::StaticClass());
 	RigVMClient.SetOuterClientHost(this, GET_MEMBER_NAME_CHECKED(UAnimNextRigVMAssetEditorData, RigVMClient));
 
+	const bool bIsDuplicating = (Ar.GetPortFlags() & PPF_Duplicate) != 0;
+	if (bIsDuplicating)
+	{
+		Ar << Entries;
+	}
+
 	Super::Serialize(Ar);
 }
 
@@ -123,6 +130,14 @@ void UAnimNextRigVMAssetEditorData::PostLoad()
 
 	// delay compilation until the package has been loaded
 	FCoreUObjectDelegates::OnEndLoadPackage.AddUObject(this, &UAnimNextRigVMAssetEditorData::HandlePackageDone);
+	
+	if (IInterface_AssetUserData* OuterUserData = Cast<IInterface_AssetUserData>(GetOuter()))
+	{
+		if(!OuterUserData->HasAssetUserDataOfClass(UAnimNextGraphWorkspaceAssetUserData::StaticClass()))
+		{
+			OuterUserData->AddAssetUserDataOfClass(UAnimNextGraphWorkspaceAssetUserData::StaticClass());
+		}
+	}
 }
 
 void UAnimNextRigVMAssetEditorData::PostLoadExternalPackages()
@@ -166,6 +181,18 @@ void UAnimNextRigVMAssetEditorData::GetAssetRegistryTags(FAssetRegistryTagsConte
 	Context.AddTag(FAssetRegistryTag(UE::AnimNext::ExportsAnimNextAssetRegistryTag, TagValue, FAssetRegistryTag::TT_Hidden));
 }
 
+bool UAnimNextRigVMAssetEditorData::Rename(const TCHAR* NewName, UObject* NewOuter, ERenameFlags Flags)
+{
+	FExternalPackageHelper::FRenameExternalObjectsHelperContext Context(this, Flags);
+	return Super::Rename(NewName, NewOuter, Flags);
+}
+
+void UAnimNextRigVMAssetEditorData::PreDuplicate(FObjectDuplicationParameters& DupParams)
+{
+	UObject::PreDuplicate(DupParams);
+	FExternalPackageHelper::DuplicateExternalPackages(this, DupParams);
+}
+
 void UAnimNextRigVMAssetEditorData::HandlePackageDone(const FEndLoadPackageContext& Context)
 {
 	if (!Context.LoadedPackages.Contains(GetPackage()))
diff --git a/Engine/Plugins/Experimental/AnimNext/Source/AnimNextUncookedOnly/Private/AnimNextRigVMWorkspaceAssetUserData.cpp b/Engine/Plugins/Experimental/AnimNext/Source/AnimNextUncookedOnly/Private/AnimNextRigVMWorkspaceAssetUserData.cpp
new file mode 100644
--- /dev/null
+++ b/Engine/Plugins/Experimental/AnimNext/Source/AnimNextUncookedOnly/Private/AnimNextRigVMWorkspaceAssetUserData.cpp
@@ -0,0 +1,33 @@
+// Copyright Epic Games, Inc. All Rights Reserved.
+
+#include "AnimNextRigVMWorkspaceAssetUserData.h"
+
+#include "UncookedOnlyUtils.h"
+#include "Graph/AnimNextGraph.h"
+#include "Graph/AnimNextGraph_EditorData.h"
+#include "UObject/AssetRegistryTagsContext.h"
+
+void UAnimNextGraphWorkspaceAssetUserData::GetAssetRegistryTags(FAssetRegistryTagsContext Context) const
+{
+	Super::GetAssetRegistryTags(Context);
+
+	FWorkspaceOutlinerItemExports Exports;
+	{
+		const UAnimNextGraph* GraphOuter = CastChecked<UAnimNextGraph>(GetOuter());
+		const UAnimNextGraph_EditorData* GraphEditorData = UE::AnimNext::UncookedOnly::FUtils::GetEditorData(GraphOuter);	
+		{
+			FWorkspaceOutlinerItemExport& RootAssetExport = Exports.Exports.AddDefaulted_GetRef();	
+			RootAssetExport.Identifier = GraphOuter->GetFName();
+			RootAssetExport.ParentIdentifier = NAME_None;
+			RootAssetExport.AssetPath = GraphOuter;
+
+			RootAssetExport.Data.InitializeAsScriptStruct(FAnimNextGraphAssetOutlinerData::StaticStruct());
+		}
+	
+		UE::AnimNext::UncookedOnly::FUtils::GetAssetOutlinerItems(GraphEditorData, Exports);
+	}
+	
+	FString TagValue;
+	FWorkspaceOutlinerItemExports::StaticStruct()->ExportText(TagValue, &Exports, nullptr, nullptr, PPF_None, nullptr);
+	Context.AddTag(FAssetRegistryTag(UE::Workspace::ExportsWorkspaceItemsRegistryTag, TagValue, FAssetRegistryTag::TT_Hidden));
+}
diff --git a/Engine/Plugins/Experimental/AnimNext/Source/AnimNextUncookedOnly/Private/AnimNextUncookedOnlyModule.cpp b/Engine/Plugins/Experimental/AnimNext/Source/AnimNextUncookedOnly/Private/AnimNextUncookedOnlyModule.cpp
--- a/Engine/Plugins/Experimental/AnimNext/Source/AnimNextUncookedOnly/Private/AnimNextUncookedOnlyModule.cpp
+++ b/Engine/Plugins/Experimental/AnimNext/Source/AnimNextUncookedOnly/Private/AnimNextUncookedOnlyModule.cpp
@@ -1,6 +1,8 @@
 // Copyright Epic Games, Inc. All Rights Reserved.
 
 #include "AnimNextUncookedOnlyModule.h"
+
+#include "AnimNextRigVMWorkspaceAssetUserData.h"
 #include "UncookedOnlyUtils.h"
 #include "Engine/Blueprint.h"
 #include "Modules/ModuleManager.h"
@@ -25,11 +27,30 @@ void FModule::StartupModule()
 	UAnimNextSchedule::GetAssetRegistryTagsFunction = [](const UAnimNextSchedule* InSchedule, FAssetRegistryTagsContext Context)
 	{
 		FAnimNextParameterProviderAssetRegistryExports Exports;
-		FUtils::GetScheduleParameters(InSchedule, Exports);
-		
-		FString TagValue;
-		FAnimNextParameterProviderAssetRegistryExports::StaticStruct()->ExportText(TagValue, &Exports, nullptr, nullptr, PPF_None, nullptr);
-		Context.AddTag(UObject::FAssetRegistryTag(UE::AnimNext::ExportsAnimNextAssetRegistryTag, TagValue, UObject::FAssetRegistryTag::TT_Hidden));
+        {
+		    FUtils::GetScheduleParameters(InSchedule, Exports);
+		    
+		    FString TagValue;
+		    FAnimNextParameterProviderAssetRegistryExports::StaticStruct()->ExportText(TagValue, &Exports, nullptr, nullptr, PPF_None, nullptr);
+		    Context.AddTag(UObject::FAssetRegistryTag(UE::AnimNext::ExportsAnimNextAssetRegistryTag, TagValue, UObject::FAssetRegistryTag::TT_Hidden));
+        }
+
+		FWorkspaceOutlinerItemExports OutlinerExports;
+		{
+			FWorkspaceOutlinerItemExport& RootAssetExport = OutlinerExports.Exports.AddDefaulted_GetRef();
+			RootAssetExport.Identifier = InSchedule->GetFName();
+			RootAssetExport.ParentIdentifier = NAME_None;
+			RootAssetExport.AssetPath = InSchedule;
+
+			RootAssetExport.Data.InitializeAsScriptStruct(FAnimNextSchedulerData::StaticStruct());
+			FAnimNextSchedulerData& AssetData = RootAssetExport.Data.GetMutable<FAnimNextSchedulerData>();
+
+			{
+				FString TagValue;
+				FWorkspaceOutlinerItemExports::StaticStruct()->ExportText(TagValue, &OutlinerExports, nullptr, nullptr, PPF_None, nullptr);
+				Context.AddTag(UObject::FAssetRegistryTag(UE::Workspace::ExportsWorkspaceItemsRegistryTag, TagValue, UObject::FAssetRegistryTag::TT_Hidden));
+			}
+		}
 	};
 	
 	// Ensure that any BP components that we care about contribute to the parameter pool
diff --git a/Engine/Plugins/Experimental/AnimNext/Source/AnimNextUncookedOnly/Private/UncookedOnlyUtils.cpp b/Engine/Plugins/Experimental/AnimNext/Source/AnimNextUncookedOnly/Private/UncookedOnlyUtils.cpp
--- a/Engine/Plugins/Experimental/AnimNext/Source/AnimNextUncookedOnly/Private/UncookedOnlyUtils.cpp
+++ b/Engine/Plugins/Experimental/AnimNext/Source/AnimNextUncookedOnly/Private/UncookedOnlyUtils.cpp
@@ -49,6 +49,7 @@
 #include "Param/AnimNextParamInstanceIdentifier.h"
 #include "Param/IParameterSourceType.h"
 #include "Param/RigVMDispatch_GetScopedParameter.h"
+#include "AnimNextRigVMWorkspaceAssetUserData.h"
 
 #define LOCTEXT_NAMESPACE "AnimNextUncookedOnlyUtils"
 
@@ -1600,6 +1601,49 @@ void FUtils::GetBlueprintParameters(const UBlueprint* InBlueprint, TSet<FAnimNex
 	}
 }
 
+void FUtils::GetAssetOutlinerItems(const UAnimNextRigVMAssetEditorData* EditorData, FWorkspaceOutlinerItemExports& OutExports)
+{
+	constexpr bool bExportParametersAsOutlinerItems = false;
+	if (bExportParametersAsOutlinerItems)
+	{
+		TSet<FName> ParameterNames;
+		
+		FAnimNextParameterProviderAssetRegistryExports GraphExports;
+		UE::AnimNext::UncookedOnly::FUtils::GetAssetParameters(EditorData, GraphExports);
+		for ( const FAnimNextParameterAssetRegistryExportEntry& Entry : GraphExports.Parameters)
+		{
+			if (!ParameterNames.Contains(Entry.Name))
+			{
+				FWorkspaceOutlinerItemExport& ParameterExport = OutExports.Exports.AddDefaulted_GetRef();	
+				ParameterExport.Identifier = Entry.Name;						
+				ParameterExport.ParentIdentifier = EditorData->GetOuter()->GetFName();
+				ParameterExport.AssetPath = EditorData->GetOuter();
+
+				ParameterExport.Data.InitializeAsScriptStruct(FAnimNextParameterOutlinerData::StaticStruct());
+				FAnimNextParameterOutlinerData& AssetData = ParameterExport.Data.GetMutable<FAnimNextParameterOutlinerData>();
+				AssetData.Type = Entry.Type;
+
+				ParameterNames.Add(Entry.Name);
+			}
+		}
+	}
+	
+	for(const UAnimNextRigVMAssetEntry* Entry : EditorData->Entries)
+	{
+		if(const IAnimNextRigVMGraphInterface* GraphInterface = Cast<IAnimNextRigVMGraphInterface>(Entry))
+		{
+			FWorkspaceOutlinerItemExport& Export = OutExports.Exports.AddDefaulted_GetRef();
+			Export.Identifier = GraphInterface->GetEdGraph()->GetFName();
+			Export.ParentIdentifier = EditorData->GetOuter()->GetFName();					
+			Export.AssetPath = EditorData->GetOuter();
+
+			Export.Data.InitializeAsScriptStruct(FAnimNextGraphOutlinerData::StaticStruct());
+			FAnimNextGraphOutlinerData& GraphData = Export.Data.GetMutable<FAnimNextGraphOutlinerData>();
+			GraphData.GraphInterface = GraphInterface->_getUObject();
+		}
+	}
+}
+
 void FUtils::CompileSchedule(UAnimNextSchedule* InSchedule)
 {
 	using namespace UE::AnimNext;
diff --git a/Engine/Plugins/Experimental/Workspace/Source/WorkspaceEditor/Private/AssetDocumentSummoner.cpp b/Engine/Plugins/Experimental/Workspace/Source/WorkspaceEditor/Private/AssetDocumentSummoner.cpp
--- a/Engine/Plugins/Experimental/Workspace/Source/WorkspaceEditor/Private/AssetDocumentSummoner.cpp
+++ b/Engine/Plugins/Experimental/Workspace/Source/WorkspaceEditor/Private/AssetDocumentSummoner.cpp
@@ -10,15 +10,17 @@
 #include "Widgets/Docking/SDockTab.h"
 #include "Widgets/Layout/SSpacer.h"
 #include "WorkspaceDocumentState.h"
+#include "SWorkspaceTabWrapper.h"
 
 #define LOCTEXT_NAMESPACE "AssetDocumentSummoner"
 
 namespace UE::Workspace
 {
 
-FAssetDocumentSummoner::FAssetDocumentSummoner(FName InIdentifier, TSharedPtr<FWorkspaceEditor> InHostingApp)
+FAssetDocumentSummoner::FAssetDocumentSummoner(FName InIdentifier, TSharedPtr<FWorkspaceEditor> InHostingApp, bool bInAllowUnsupportedClasses /*= false*/)
 	: FDocumentTabFactoryForObjects<UObject>(InIdentifier, InHostingApp)
 	, HostingAppPtr(InHostingApp)
+	, bAllowUnsupportedClasses(bInAllowUnsupportedClasses)
 {
 }
 
@@ -46,11 +48,12 @@ void FAssetDocumentSummoner::SaveState(TSharedPtr<SDockTab> Tab, TSharedPtr<FTab
 		UObject* Object = Payload->IsValid() ? FTabPayload_UObject::CastChecked<UObject>(Payload) : nullptr;
 		if(Object)
 		{
+			TSharedRef<SWorkspaceTabWrapper> TabWrapper = StaticCastSharedRef<SWorkspaceTabWrapper>(Tab->GetContent());			
 			FWorkspaceEditorModule& WorkspaceEditorModule = FModuleManager::LoadModuleChecked<FWorkspaceEditorModule>("WorkspaceEditor");
 			const FObjectDocumentArgs* DocumentArgs = WorkspaceEditorModule.FindObjectDocumentType(Object->GetClass()->GetClassPathName());
 			if(DocumentArgs && DocumentArgs->OnGetDocumentState.IsBound())
 			{
-				WorkspaceEditor->RecordDocumentState(DocumentArgs->OnGetDocumentState.Execute(FWorkspaceEditorContext(WorkspaceEditor.ToSharedRef(), Object), Tab->GetContent()));
+				WorkspaceEditor->RecordDocumentState(DocumentArgs->OnGetDocumentState.Execute(FWorkspaceEditorContext(WorkspaceEditor.ToSharedRef(), Object), TabWrapper->GetContent()));
 			}
 			else
 			{
@@ -96,7 +99,7 @@ bool FAssetDocumentSummoner::IsPayloadSupported(TSharedRef<FTabPayload> Payload)
 		FWorkspaceEditorModule& WorkspaceEditorModule = FModuleManager::LoadModuleChecked<FWorkspaceEditorModule>("WorkspaceEditor");
 		if(const FObjectDocumentArgs* DocumentArgs = WorkspaceEditorModule.FindObjectDocumentType(Object->GetClass()->GetClassPathName()))
 		{
-			return AllowedClassPaths.IsEmpty() || AllowedClassPaths.Contains(Object->GetClass()->GetClassPathName());
+			return AllowedClassPaths.Contains(Object->GetClass()->GetClassPathName()) || bAllowUnsupportedClasses;
 		}
 	}
 
@@ -127,19 +130,24 @@ TAttribute<FText> FAssetDocumentSummoner::ConstructTabLabelSuffix(const FWorkflo
 
 TSharedRef<SWidget> FAssetDocumentSummoner::CreateTabBodyForObject(const FWorkflowTabSpawnInfo& Info, UObject* DocumentID) const
 {
-	if(TSharedPtr<FWorkspaceEditor> WorkspaceEditor = HostingAppPtr.Pin())
+	TSharedPtr<SWidget> TabContent = SNullWidget::NullWidget;
+
+	const TSharedPtr<FWorkspaceEditor> WorkspaceEditor = HostingAppPtr.Pin();
+	check(WorkspaceEditor.IsValid());
+	
+	const FWorkspaceEditorModule& WorkspaceEditorModule = FModuleManager::LoadModuleChecked<FWorkspaceEditorModule>("WorkspaceEditor");
+	if(const FObjectDocumentArgs* DocumentArgs = WorkspaceEditorModule.FindObjectDocumentType(DocumentID->GetClass()->GetClassPathName()))
 	{
-		FWorkspaceEditorModule& WorkspaceEditorModule = FModuleManager::LoadModuleChecked<FWorkspaceEditorModule>("WorkspaceEditor");
-		if(const FObjectDocumentArgs* DocumentArgs = WorkspaceEditorModule.FindObjectDocumentType(DocumentID->GetClass()->GetClassPathName()))
+		if(DocumentArgs->OnMakeDocumentWidget.IsBound())
 		{
-			if(DocumentArgs->OnMakeDocumentWidget.IsBound())
-			{
-				return DocumentArgs->OnMakeDocumentWidget.Execute(FWorkspaceEditorContext(WorkspaceEditor.ToSharedRef(), DocumentID));
-			}
+			TabContent = DocumentArgs->OnMakeDocumentWidget.Execute(FWorkspaceEditorContext(WorkspaceEditor.ToSharedRef(), DocumentID));
 		}
 	}
 
-	return SNew(SSpacer);
+	return SNew(SWorkspaceTabWrapper, Info.TabInfo, WorkspaceEditor, DocumentID)
+	[
+		TabContent.ToSharedRef()
+	];
 }
 
 const FSlateBrush* FAssetDocumentSummoner::GetTabIconForObject(const FWorkflowTabSpawnInfo& Info, UObject* DocumentID) const
diff --git a/Engine/Plugins/Experimental/Workspace/Source/WorkspaceEditor/Private/AssetDocumentSummoner.h b/Engine/Plugins/Experimental/Workspace/Source/WorkspaceEditor/Private/AssetDocumentSummoner.h
--- a/Engine/Plugins/Experimental/Workspace/Source/WorkspaceEditor/Private/AssetDocumentSummoner.h
+++ b/Engine/Plugins/Experimental/Workspace/Source/WorkspaceEditor/Private/AssetDocumentSummoner.h
@@ -22,7 +22,7 @@ public:
 	// Delegate called to save the state of a document
 	DECLARE_DELEGATE_OneParam(FOnSaveDocumentState, UObject*);
 
-	FAssetDocumentSummoner(FName InIdentifier, TSharedPtr<FWorkspaceEditor> InHostingApp);
+	FAssetDocumentSummoner(FName InIdentifier, TSharedPtr<FWorkspaceEditor> InHostingApp, bool bInAllowUnsupportedClasses = false);
 
 	void SetAllowedClassPaths(TConstArrayView<FTopLevelAssetPath> InAllowedClassPaths);
 
@@ -48,6 +48,9 @@ private:
 
 	// Allowed object types
 	TArray<FTopLevelAssetPath> AllowedClassPaths;
+
+	// Whether or not to allow objects if AllowedClassPaths does not contain their class
+	bool bAllowUnsupportedClasses = false;
 };
 
 }
diff --git a/Engine/Plugins/Experimental/Workspace/Source/WorkspaceEditor/Private/Outliner/WorkspaceOutlinerColumns.cpp b/Engine/Plugins/Experimental/Workspace/Source/WorkspaceEditor/Private/Outliner/WorkspaceOutlinerColumns.cpp
new file mode 100644
--- /dev/null
+++ b/Engine/Plugins/Experimental/Workspace/Source/WorkspaceEditor/Private/Outliner/WorkspaceOutlinerColumns.cpp
@@ -0,0 +1,97 @@
+// Copyright Epic Games, Inc. All Rights Reserved.
+
+#include "WorkspaceOutlinerColumns.h"
+
+#include "ISourceControlModule.h"
+#include "ISourceControlState.h"
+#include "WorkspaceEditorModule.h"
+#include "WorkspaceOutlinerTreeItem.h"
+#include "RevisionControlStyle/RevisionControlStyle.h"
+#include "ISourceControlProvider.h"
+#include "SourceControlHelpers.h"
+#include "Widgets/Images/SImage.h"
+
+#define LOCTEXT_NAMESPACE "WorkspaceOutlinerColumns"
+
+namespace UE::Workspace
+{
+	FName WorkspaceOutlinerSourceControl("Source Control Status");
+	FName FWorkspaceOutlinerSourceControlColumn::GetID()
+	{
+		return WorkspaceOutlinerSourceControl;
+	}
+
+	SHeaderRow::FColumn::FArguments FWorkspaceOutlinerSourceControlColumn::ConstructHeaderRowColumn()
+	{
+		return SHeaderRow::Column(GetColumnID())
+		.FixedWidth(24.f)
+		.HAlignHeader(HAlign_Center)
+		.VAlignHeader(VAlign_Center)
+		.HAlignCell(HAlign_Center)
+		.VAlignCell(VAlign_Center)
+		[
+			SNew(SBox)
+			.WidthOverride(16.0f)
+			.HeightOverride(16.0f)
+			.VAlign(VAlign_Center)
+			.HAlign(HAlign_Center)
+			[
+				SNew(SImage)
+				.ColorAndOpacity(FSlateColor::UseForeground())
+				.Image(FRevisionControlStyleManager::Get().GetBrush("RevisionControl.Icon"))
+				.ToolTipText(LOCTEXT("RevisionControlStatusTooltip", "Revision control status of this parameter"))
+			]
+		];	
+	}
+
+	const TSharedRef<SWidget> FWorkspaceOutlinerSourceControlColumn::ConstructRowWidget(FSceneOutlinerTreeItemRef Item, const STableRow<FSceneOutlinerTreeItemPtr>& Row)
+	{
+		if (const FWorkspaceOutlinerTreeItem* TreeItem = Item->CastTo<FWorkspaceOutlinerTreeItem>())
+		{
+			FString PackagePath;
+			if (const TSharedPtr<IWorkspaceOutlinerItemDetails> SharedFactory = FWorkspaceEditorModule::GetOutlinerItemDetails(MakeOutlinerDetailsId(TreeItem->Export)))
+			{
+				PackagePath = SourceControlHelpers::PackageFilename(SharedFactory->GetPackage(TreeItem->Export));
+			}
+			else if (TreeItem->Export.ParentIdentifier == NAME_None && TreeItem->Export.AssetPath.IsValid())
+			{
+				PackagePath = SourceControlHelpers::PackageFilename(TreeItem->Export.AssetPath.GetLongPackageName());
+			}
+		
+			if (!PackagePath.IsEmpty())
+			{
+				return SNew(SBox)
+				.HAlign(HAlign_Center)
+				.VAlign(VAlign_Center)
+				.HeightOverride(20.0f)
+				[
+					SNew(SImage)
+					.ToolTipText_Lambda([this, PackagePath]() -> FText
+					{
+						ISourceControlProvider& SourceControlProvider = ISourceControlModule::Get().GetProvider();
+						if(const FSourceControlStatePtr State = SourceControlProvider.GetState(PackagePath, EStateCacheUsage::Use))
+						{
+							return FText::Format(LOCTEXT("RevisionControlStatusFormat", "File: {0}\nStatus: {1}"), State->GetDisplayName(), State->GetDisplayTooltip());
+						}
+
+						return LOCTEXT("RevisionControlStatus", "Revision control status of this parameter");
+					})
+					.Image_Lambda([this, PackagePath]() -> const FSlateBrush*
+					{
+						ISourceControlProvider& SourceControlProvider = ISourceControlModule::Get().GetProvider();
+						if(const FSourceControlStatePtr State = SourceControlProvider.GetState(PackagePath, EStateCacheUsage::Use))
+						{
+							return State->GetIcon().GetSmallIcon();
+						}
+						
+						return nullptr;
+					})
+				];
+			}
+		}
+		
+		return SNullWidget::NullWidget;
+	}
+}
+
+#undef LOCTEXT_NAMESPACE // "WorkspaceOutlinerColumns"
\ No newline at end of file
diff --git a/Engine/Plugins/Experimental/Workspace/Source/WorkspaceEditor/Private/Outliner/WorkspaceOutlinerColumns.h b/Engine/Plugins/Experimental/Workspace/Source/WorkspaceEditor/Private/Outliner/WorkspaceOutlinerColumns.h
new file mode 100644
--- /dev/null
+++ b/Engine/Plugins/Experimental/Workspace/Source/WorkspaceEditor/Private/Outliner/WorkspaceOutlinerColumns.h
@@ -0,0 +1,34 @@
+// Copyright Epic Games, Inc. All Rights Reserved.
+
+#pragma once
+
+#include "ISceneOutliner.h"
+#include "ISceneOutlinerColumn.h"
+#include "SceneOutlinerFwd.h"
+#include "Templates/SharedPointer.h"
+#include "UObject/NameTypes.h"
+#include "Widgets/Views/SHeaderRow.h"
+
+class SWidget;
+
+namespace UE::Workspace
+{
+class FWorkspaceOutlinerSourceControlColumn : public ISceneOutlinerColumn
+{
+public:
+	static FName GetID();
+	
+	FWorkspaceOutlinerSourceControlColumn(ISceneOutliner& SceneOutliner) : WeakSceneOutliner(StaticCastSharedRef<ISceneOutliner>(SceneOutliner.AsShared())) {}
+	virtual ~FWorkspaceOutlinerSourceControlColumn() override = default;
+	
+	// Begin ISceneOutlinerColumn overrides
+	virtual FName GetColumnID() override { return GetID(); }
+	virtual SHeaderRow::FColumn::FArguments ConstructHeaderRowColumn() override;
+	virtual const TSharedRef<SWidget> ConstructRowWidget(FSceneOutlinerTreeItemRef TreeItem, const STableRow<FSceneOutlinerTreeItemPtr>& Row) override;
+	virtual bool SupportsSorting() const override { return false; }
+	// End ISceneOutlinerColumn overrides
+
+private:
+	TWeakPtr<ISceneOutliner> WeakSceneOutliner;
+};
+}
diff --git a/Engine/Plugins/Experimental/Workspace/Source/WorkspaceEditor/Private/Outliner/WorkspaceOutlinerHierarchy.cpp b/Engine/Plugins/Experimental/Workspace/Source/WorkspaceEditor/Private/Outliner/WorkspaceOutlinerHierarchy.cpp
new file mode 100644
--- /dev/null
+++ b/Engine/Plugins/Experimental/Workspace/Source/WorkspaceEditor/Private/Outliner/WorkspaceOutlinerHierarchy.cpp
@@ -0,0 +1,84 @@
+// Copyright Epic Games, Inc. All Rights Reserved.
+
+#include "WorkspaceOutlinerHierarchy.h"
+
+#include "ISceneOutlinerMode.h"
+#include "Workspace.h"
+#include "WorkspaceAssetRegistryInfo.h"
+#include "WorkspaceOutlinerTreeItem.h"
+
+namespace UE::Workspace
+{
+	FWorkspaceOutlinerHierarchy::FWorkspaceOutlinerHierarchy(ISceneOutlinerMode* Mode, const TWeakObjectPtr<UWorkspace>& InWorkspace) : ISceneOutlinerHierarchy(Mode), WeakWorkspace(InWorkspace)
+	{
+	}
+
+	void FWorkspaceOutlinerHierarchy::CreateItems(TArray<FSceneOutlinerTreeItemPtr>& OutItems) const
+	{
+		if (const UWorkspace* Workspace = WeakWorkspace.Get())
+		{
+			TArray<FAssetData> AssetDataEntries; 
+			Workspace->GetAssetDataEntries(AssetDataEntries);
+
+			for (const FAssetData& AssetData : AssetDataEntries)
+			{
+				FString TagValue;
+				if(AssetData.GetTagValue(UE::Workspace::ExportsWorkspaceItemsRegistryTag, TagValue))
+				{
+					FWorkspaceOutlinerItemExports Exports;
+					FWorkspaceOutlinerItemExports::StaticStruct()->ImportText(*TagValue, &Exports, nullptr, 0, nullptr, FWorkspaceOutlinerItemExports::StaticStruct()->GetName());
+					for (const FWorkspaceOutlinerItemExport& Export : Exports.Exports)
+					{
+						if (FSceneOutlinerTreeItemPtr Item = Mode->CreateItemFor<FWorkspaceOutlinerTreeItem>(FWorkspaceOutlinerTreeItem::FItemData(Export)))
+						{
+							OutItems.Add(Item);
+						}
+					}
+				}
+			}
+		}		
+	}
+
+	FSceneOutlinerTreeItemPtr FWorkspaceOutlinerHierarchy::FindOrCreateParentItem(const ISceneOutlinerTreeItem& Item, const TMap<FSceneOutlinerTreeItemID, FSceneOutlinerTreeItemPtr>& Items, bool bCreate)
+	{
+		if (const FWorkspaceOutlinerTreeItem* TreeItem = Item.CastTo<FWorkspaceOutlinerTreeItem>())
+		{
+			const FName ParentIdentifier = TreeItem->Export.ParentIdentifier;
+			if (ParentIdentifier != NAME_None)
+			{
+				if (const FSceneOutlinerTreeItemPtr* ParentItem = Items.Find(HashCombine(GetTypeHash(TreeItem->Export.AssetPath),GetTypeHash(ParentIdentifier))))
+				{
+					return *ParentItem;
+				}
+				else if(bCreate)
+				{
+					if (const UWorkspace* Workspace = WeakWorkspace.Get())
+					{
+						TArray<FAssetData> AssetDataEntries; 
+						Workspace->GetAssetDataEntries(AssetDataEntries);
+
+						if(const FAssetData* AssetDataPtr = AssetDataEntries.FindByPredicate([AssetPath = TreeItem->Export.AssetPath](const FAssetData& AssetData) { return AssetData.GetSoftObjectPath() == AssetPath; }))
+						{
+							FString TagValue;
+							if((*AssetDataPtr).GetTagValue(UE::Workspace::ExportsWorkspaceItemsRegistryTag, TagValue))
+							{
+								FWorkspaceOutlinerItemExports Exports;
+								FWorkspaceOutlinerItemExports::StaticStruct()->ImportText(*TagValue, &Exports, nullptr, 0, nullptr, FWorkspaceOutlinerItemExports::StaticStruct()->GetName());
+								
+								if (const FWorkspaceOutlinerItemExport* ExportPtr = Exports.Exports.FindByPredicate([ParentIdentifier](const FWorkspaceOutlinerItemExport& ItemExport)
+								{
+									return ItemExport.Identifier == ParentIdentifier;
+								}))
+								{
+									Mode->CreateItemFor<FWorkspaceOutlinerTreeItem>(FWorkspaceOutlinerTreeItem::FItemData(*ExportPtr), true);
+								}
+							}
+						}
+					}
+				}
+			}
+		}
+		
+		return nullptr;
+	}
+}
diff --git a/Engine/Plugins/Experimental/Workspace/Source/WorkspaceEditor/Private/Outliner/WorkspaceOutlinerHierarchy.h b/Engine/Plugins/Experimental/Workspace/Source/WorkspaceEditor/Private/Outliner/WorkspaceOutlinerHierarchy.h
new file mode 100644
--- /dev/null
+++ b/Engine/Plugins/Experimental/Workspace/Source/WorkspaceEditor/Private/Outliner/WorkspaceOutlinerHierarchy.h
@@ -0,0 +1,27 @@
+// Copyright Epic Games, Inc. All Rights Reserved.
+
+#pragma once
+
+#include "CoreMinimal.h"
+#include "ISceneOutlinerHierarchy.h"
+
+class UWorkspace;
+
+namespace UE::Workspace
+{
+	class FWorkspaceOutlinerHierarchy : public ISceneOutlinerHierarchy
+	{
+	public:
+		FWorkspaceOutlinerHierarchy(ISceneOutlinerMode* Mode, const TWeakObjectPtr<UWorkspace>& InWorkspace);
+		FWorkspaceOutlinerHierarchy(const FWorkspaceOutlinerHierarchy&) = delete;
+		FWorkspaceOutlinerHierarchy& operator=(const FWorkspaceOutlinerHierarchy&) = delete;
+
+		// Begin ISceneOutlinerHierarchy overrides
+		virtual void CreateItems(TArray<FSceneOutlinerTreeItemPtr>& OutItems) const override;
+		virtual void CreateChildren(const FSceneOutlinerTreeItemPtr& Item, TArray<FSceneOutlinerTreeItemPtr>& OutChildren) const override {}
+		virtual FSceneOutlinerTreeItemPtr FindOrCreateParentItem(const ISceneOutlinerTreeItem& Item, const TMap<FSceneOutlinerTreeItemID, FSceneOutlinerTreeItemPtr>& Items, bool bCreate = false) override;
+		// End ISceneOutlinerHierarchy overrides
+	private:
+		TWeakObjectPtr<UWorkspace> WeakWorkspace;
+	};
+}
\ No newline at end of file
diff --git a/Engine/Plugins/Experimental/Workspace/Source/WorkspaceEditor/Private/Outliner/WorkspaceOutlinerMode.cpp b/Engine/Plugins/Experimental/Workspace/Source/WorkspaceEditor/Private/Outliner/WorkspaceOutlinerMode.cpp
new file mode 100644
--- /dev/null
+++ b/Engine/Plugins/Experimental/Workspace/Source/WorkspaceEditor/Private/Outliner/WorkspaceOutlinerMode.cpp
@@ -0,0 +1,273 @@
+// Copyright Epic Games, Inc. All Rights Reserved.
+
+#include "WorkspaceOutlinerMode.h"
+
+#include "WorkspaceItemMenuContext.h"
+#include "WorkspaceOutlinerHierarchy.h"
+#include "WorkspaceOutlinerTreeItem.h"
+#include "Toolkits/AssetEditorToolkitMenuContext.h"
+#include "Workspace.h"
+#include "IWorkspaceEditor.h"
+#include "WorkspaceEditorModule.h"
+#include "AssetRegistry/AssetRegistryModule.h"
+#include "ScopedTransaction.h"
+
+#define LOCTEXT_NAMESPACE "FWorkspaceOutlinerMode"
+
+namespace UE::Workspace
+{
+FWorkspaceOutlinerMode::FWorkspaceOutlinerMode(SSceneOutliner* InSceneOutliner, const TWeakObjectPtr<UWorkspace>& InWeakWorkspace, const TWeakPtr<IWorkspaceEditor>& InWeakWorkspaceEditor) : ISceneOutlinerMode(InSceneOutliner), WeakWorkspace(InWeakWorkspace), WeakWorkspaceEditor(InWeakWorkspaceEditor)
+{
+	if (UWorkspace* Workspace = WeakWorkspace.Get())
+	{
+		Workspace->ModifiedDelegate.AddRaw(this, &FWorkspaceOutlinerMode::OnWorkspaceModified);
+	}
+}
+
+FWorkspaceOutlinerMode::~FWorkspaceOutlinerMode()
+{
+	if (UWorkspace* Workspace = WeakWorkspace.Get())
+	{
+		Workspace->ModifiedDelegate.RemoveAll(this);
+	}
+}
+
+void FWorkspaceOutlinerMode::Rebuild()
+{
+	Hierarchy = CreateHierarchy();
+}
+
+TSharedPtr<SWidget> FWorkspaceOutlinerMode::CreateContextMenu()
+{
+	static const FName MenuName("WorkspaceOutliner.ItemContextMenu");
+	UToolMenus* ToolMenus = UToolMenus::Get();
+	if (!ToolMenus->IsMenuRegistered(MenuName))
+	{
+		FToolMenuOwnerScoped ToolMenuOwnerScope(this);
+		if (UToolMenu* Menu = ToolMenus->RegisterMenu(MenuName))
+		{
+			Menu->AddDynamicSection(TEXT("Assets"), FNewToolMenuDelegate::CreateLambda([](UToolMenu* InMenu)
+			{
+				const UAssetEditorToolkitMenuContext* EditorContext = InMenu->FindContext<UAssetEditorToolkitMenuContext>();
+				const UWorkspaceItemMenuContext* MenuContext = InMenu->FindContext<UWorkspaceItemMenuContext>();
+				if(EditorContext && MenuContext)
+				{
+					const bool bSelectionContainsTopLevelAsset = MenuContext->SelectedExports.Num() && MenuContext->SelectedExports.ContainsByPredicate([](const FWorkspaceOutlinerItemExport& Export)
+					{
+						return Export.ParentIdentifier == NAME_None;
+					});
+
+					if (bSelectionContainsTopLevelAsset)
+					{
+						FToolMenuSection& AssetsSection = InMenu->AddSection("Assets", LOCTEXT("AssetSectionLabel", "Assets"));
+						AssetsSection.AddMenuEntry(TEXT("OpenAsset"),
+							LOCTEXT("OpenAssetLabel", "Open Asset(s)"),
+							LOCTEXT("OpenAssetTooltip", "Opens the selected Asset(s)"),
+							FSlateIcon(FAppStyle::Get().GetStyleSetName(), "SystemWideCommands.SummonOpenAssetDialog"),
+							FUIAction(FExecuteAction::CreateLambda([WeakEditor=EditorContext->Toolkit, SelectedExports=MenuContext->SelectedExports]()
+							{
+								if (const TSharedPtr<UE::Workspace::IWorkspaceEditor> SharedWorkspaceEditor = StaticCastSharedPtr<UE::Workspace::IWorkspaceEditor>(WeakEditor.Pin()))
+								{
+									TSet<FSoftObjectPath> AssetPaths;
+									for (const FWorkspaceOutlinerItemExport& ItemExport : SelectedExports)
+									{
+										if (ItemExport.ParentIdentifier == NAME_None)
+										{
+											AssetPaths.Add(ItemExport.AssetPath);	
+										}
+									}
+
+									for (const FSoftObjectPath& AssetPath : AssetPaths)
+									{
+										SharedWorkspaceEditor->OpenAssets({AssetPath.TryLoad()});	
+									}
+								}
+							}))
+						);
+
+						AssetsSection.AddMenuEntry(TEXT("RemoveAsset"),
+							LOCTEXT("RemoveAssetLabel", "Remove Asset(s)"),
+							LOCTEXT("RemoveAssetTooltip", "Removes the selected Asset(s) from the Workspace"),
+							FSlateIcon(FAppStyle::Get().GetStyleSetName(), "Icons.Delete"),
+							FUIAction(FExecuteAction::CreateLambda([WeakEditor=EditorContext->Toolkit, SelectedExports=MenuContext->SelectedExports, EditingObjects=EditorContext->GetEditingObjects()]()
+							{
+								if (const TSharedPtr<UE::Workspace::IWorkspaceEditor> SharedWorkspaceEditor = StaticCastSharedPtr<UE::Workspace::IWorkspaceEditor>(WeakEditor.Pin()))
+								{
+									TSet<FSoftObjectPath> AssetPaths;
+									for (const FWorkspaceOutlinerItemExport& ItemExport : SelectedExports)
+									{
+										if (ItemExport.ParentIdentifier == NAME_None)
+										{
+											AssetPaths.Add(ItemExport.AssetPath);	
+										}
+									}
+
+									if (EditingObjects.Num() == 1)
+									{
+										if (UWorkspace* Workspace = Cast<UWorkspace>(EditingObjects[0]))
+										{
+											for (const FSoftObjectPath& AssetPath : AssetPaths)
+											{
+												Workspace->RemoveAsset(AssetPath.TryLoad());
+											}
+										}
+									}
+								}
+							}))
+						);
+					}
+				}
+			}));
+		}
+	}
+
+	{
+		UWorkspaceItemMenuContext* MenuContext = NewObject<UWorkspaceItemMenuContext>();
+
+		Algo::TransformIf(SceneOutliner->GetSelectedItems(), MenuContext->SelectedExports,
+			[](const FSceneOutlinerTreeItemPtr& InItem)
+			{
+				return InItem->IsA<FWorkspaceOutlinerTreeItem>();
+			},
+			[](const FSceneOutlinerTreeItemPtr& InItem)
+			{
+				return InItem->CastTo<FWorkspaceOutlinerTreeItem>()->Export;
+			});	
+		
+		FToolMenuContext Context;
+		Context.AddObject(MenuContext);
+
+		WeakWorkspaceEditor.Pin()->InitToolMenuContext(Context);
+		
+		return UToolMenus::Get()->GenerateWidget(MenuName, Context);
+	}
+}
+
+void FWorkspaceOutlinerMode::OnItemDoubleClick(FSceneOutlinerTreeItemPtr Item)
+{
+	OpenItems({Item} );
+}
+
+void FWorkspaceOutlinerMode::OnItemClicked(FSceneOutlinerTreeItemPtr Item)
+{
+	const FSceneOutlinerItemSelection& Selection = SceneOutliner->GetSelection();
+	HandleItemSelection(Selection);
+}
+
+FReply FWorkspaceOutlinerMode::OnKeyDown(const FKeyEvent& InKeyEvent)
+{
+	// TODO JDB these could be in a FUICommandList
+	if (InKeyEvent.GetKey() == EKeys::Enter)
+	{
+		OpenItems(SceneOutliner->GetSelectedItems());
+		return FReply::Handled();
+	}
+	else if (InKeyEvent.GetKey() == EKeys::Platform_Delete)
+	{
+		DeleteItems(SceneOutliner->GetSelectedItems());
+		return FReply::Handled();
+	}
+
+	// TODO JDB more possible actions? (find in content browser?)
+	
+	return ISceneOutlinerMode::OnKeyDown(InKeyEvent);
+}
+
+void FWorkspaceOutlinerMode::HandleItemSelection(const FSceneOutlinerItemSelection& Selection)
+{
+	if (Selection.Num() == 1)
+	{
+		TArray<FSceneOutlinerTreeItemPtr> SelectedItems;
+		Selection.Get(SelectedItems);			
+		
+		if (TSharedPtr<IWorkspaceEditor> SharedWorkspaceEditor = WeakWorkspaceEditor.Pin())
+		{
+			if(FWorkspaceOutlinerTreeItem* TreeItem = SelectedItems[0]->CastTo<FWorkspaceOutlinerTreeItem>())
+			{
+				if(TreeItem->Export.Data.IsValid())
+				{
+					TSharedPtr<FStructOnScope> ExportDataView = MakeShared<FStructOnScope>(TreeItem->Export.Data.GetScriptStruct(), TreeItem->Export.Data.GetMutableMemory());
+					// TODO JDB handle struct selections
+					//SharedWorkspaceEditor->SetDetailsStruct(ExportDataView);
+				}
+			}
+		}
+	}
+}
+
+void FWorkspaceOutlinerMode::OnItemSelectionChanged(FSceneOutlinerTreeItemPtr Item, ESelectInfo::Type SelectionType, const FSceneOutlinerItemSelection& Selection)
+{
+	HandleItemSelection(Selection);
+}
+
+TUniquePtr<ISceneOutlinerHierarchy> FWorkspaceOutlinerMode::CreateHierarchy()
+{
+	return MakeUnique<FWorkspaceOutlinerHierarchy>(this, WeakWorkspace);
+}
+
+void FWorkspaceOutlinerMode::OnWorkspaceModified(UWorkspace* InWorkspace)
+{
+	ensure(InWorkspace == WeakWorkspace.Get());
+	SceneOutliner->FullRefresh();
+}
+
+void FWorkspaceOutlinerMode::OpenItems(TArrayView<const FSceneOutlinerTreeItemPtr> Items) const
+{
+	for (const FSceneOutlinerTreeItemPtr& Item : Items)
+	{
+		if (const FWorkspaceOutlinerTreeItem* TreeItem = Item->CastTo<FWorkspaceOutlinerTreeItem>())
+		{
+			if(TreeItem->Export.ParentIdentifier == NAME_None)
+			{
+				if (const TSharedPtr<UE::Workspace::IWorkspaceEditor> SharedWorkspaceEditor = StaticCastSharedPtr<UE::Workspace::IWorkspaceEditor>(WeakWorkspaceEditor.Pin()))
+				{			
+					SharedWorkspaceEditor->OpenAssets({TreeItem->Export.AssetPath.TryLoad()});
+				}
+			}
+			else
+			{
+				if (const TSharedPtr<IWorkspaceOutlinerItemDetails> SharedDetails = FWorkspaceEditorModule::GetOutlinerItemDetails(MakeOutlinerDetailsId(TreeItem->Export)))
+				{
+					UWorkspaceItemMenuContext* MenuContext = NewObject<UWorkspaceItemMenuContext>();
+					MenuContext->SelectedExports.Add(TreeItem->Export);
+				
+					FToolMenuContext Context(MenuContext);
+					WeakWorkspaceEditor.Pin()->InitToolMenuContext(Context);
+				
+					SharedDetails->HandleDoubleClick(Context);
+				}
+			}
+		}
+	}
+}
+
+void FWorkspaceOutlinerMode::DeleteItems(TArrayView<const FSceneOutlinerTreeItemPtr> Items) const
+{
+	FARFilter Filter;	
+	for (const FSceneOutlinerTreeItemPtr& Item : Items)
+	{
+		if (const FWorkspaceOutlinerTreeItem* TreeItem = Item->CastTo<FWorkspaceOutlinerTreeItem>())
+		{
+			if(TreeItem->Export.ParentIdentifier == NAME_None)
+			{
+				Filter.SoftObjectPaths.AddUnique(TreeItem->Export.AssetPath);
+			}
+		}		
+	}
+
+	if (UWorkspace* Workspace = WeakWorkspace.Get())
+	{
+		TArray<FAssetData> AssetDataEntriesToRemove;
+		const IAssetRegistry& AssetRegistry = FAssetRegistryModule::GetRegistry();
+		AssetRegistry.GetAssets(Filter, AssetDataEntriesToRemove);
+		if (AssetDataEntriesToRemove.Num())
+		{
+			FScopedTransaction Transaction(LOCTEXT("RemoveAssets", "Remove assets from workspace"));
+			Workspace->RemoveAssets(AssetDataEntriesToRemove);
+		}		
+	}
+}
+
+}	
+
+#undef LOCTEXT_NAMESPACE // "FWorkspaceOutlinerMode"
\ No newline at end of file
diff --git a/Engine/Plugins/Experimental/Workspace/Source/WorkspaceEditor/Private/Outliner/WorkspaceOutlinerMode.h b/Engine/Plugins/Experimental/Workspace/Source/WorkspaceEditor/Private/Outliner/WorkspaceOutlinerMode.h
new file mode 100644
--- /dev/null
+++ b/Engine/Plugins/Experimental/Workspace/Source/WorkspaceEditor/Private/Outliner/WorkspaceOutlinerMode.h
@@ -0,0 +1,40 @@
+// Copyright Epic Games, Inc. All Rights Reserved.
+
+#pragma once
+
+#include "ISceneOutlinerMode.h"
+
+class UWorkspace;
+
+namespace UE::Workspace
+{
+	class IWorkspaceEditor;
+
+	class FWorkspaceOutlinerMode : public ISceneOutlinerMode
+	{
+	public:
+		FWorkspaceOutlinerMode(SSceneOutliner* InSceneOutliner, const TWeakObjectPtr<UWorkspace>& InWeakWorkspace, const TWeakPtr<IWorkspaceEditor>& InWeakWorkspaceEditor);
+		virtual ~FWorkspaceOutlinerMode() override;
+
+		// Begin ISceneOutlinerMode overrides
+		virtual void Rebuild() override;
+		virtual TSharedPtr<SWidget> CreateContextMenu() override;		
+		virtual void OnItemDoubleClick(FSceneOutlinerTreeItemPtr Item) override;
+		virtual void OnItemClicked(FSceneOutlinerTreeItemPtr Item) override;
+		virtual FReply OnKeyDown(const FKeyEvent& InKeyEvent) override;
+		void HandleItemSelection(const FSceneOutlinerItemSelection& Selection);
+
+		virtual void OnItemSelectionChanged(FSceneOutlinerTreeItemPtr Item, ESelectInfo::Type SelectionType, const FSceneOutlinerItemSelection& Selection) override;
+	protected:
+		virtual TUniquePtr<ISceneOutlinerHierarchy> CreateHierarchy() override;		
+		// End ISceneOutlinerMode overrides
+
+		void OnWorkspaceModified(UWorkspace* InWorkspace);
+
+		void OpenItems(TArrayView<const FSceneOutlinerTreeItemPtr> Items) const;
+		void DeleteItems(TArrayView<const FSceneOutlinerTreeItemPtr> Items) const;
+	private:
+		TWeakObjectPtr<UWorkspace> WeakWorkspace;
+		TWeakPtr<IWorkspaceEditor> WeakWorkspaceEditor;
+	};
+}
\ No newline at end of file
diff --git a/Engine/Plugins/Experimental/Workspace/Source/WorkspaceEditor/Private/Outliner/WorkspaceOutlinerTreeItem.cpp b/Engine/Plugins/Experimental/Workspace/Source/WorkspaceEditor/Private/Outliner/WorkspaceOutlinerTreeItem.cpp
new file mode 100644
--- /dev/null
+++ b/Engine/Plugins/Experimental/Workspace/Source/WorkspaceEditor/Private/Outliner/WorkspaceOutlinerTreeItem.cpp
@@ -0,0 +1,119 @@
+// Copyright Epic Games, Inc. All Rights Reserved.
+
+#include "WorkspaceOutlinerTreeItem.h"
+
+#include "ISceneOutliner.h"
+#include "Styling/SlateColor.h"
+#include "Misc/Attribute.h"
+#include "Misc/Optional.h"
+#include "SceneOutlinerStandaloneTypes.h"
+#include "WorkspaceEditorModule.h"
+#include "Widgets/SCompoundWidget.h"
+#include "Widgets/DeclarativeSyntaxSupport.h"
+#include "Widgets/Images/SImage.h"
+#include "Widgets/Text/STextBlock.h"
+#include "Widgets/Views/STableRow.h"
+#include "Widgets/Views/SListView.h"
+
+#define LOCTEXT_NAMESPACE "WorkspaceOutlinerTreeItem"
+
+namespace UE::Workspace
+{
+	const FSceneOutlinerTreeItemType FWorkspaceOutlinerTreeItem::Type;
+	class SWorkspaceOutlinerTreelabel : FSceneOutlinerCommonLabelData, public SCompoundWidget
+	{
+		SLATE_BEGIN_ARGS(SWorkspaceOutlinerTreelabel) {}
+		SLATE_END_ARGS()
+
+		void Construct(const FArguments& InArgs, FWorkspaceOutlinerTreeItem& InTreeItem, ISceneOutliner& SceneOutliner, const STableRow<FSceneOutlinerTreeItemPtr>& InRow)
+		{
+			WeakSceneOutliner = StaticCastSharedRef<ISceneOutliner>(SceneOutliner.AsShared());
+			TreeItem = StaticCastSharedRef<FWorkspaceOutlinerTreeItem>(InTreeItem.AsShared());
+
+			ChildSlot
+			[
+				SNew(SHorizontalBox)
+				+ SHorizontalBox::Slot()
+				.AutoWidth()
+				.VAlign(VAlign_Center)
+				.Padding(FSceneOutlinerDefaultTreeItemMetrics::IconPadding())
+				[
+					SNew(SBox)
+					.WidthOverride(FSceneOutlinerDefaultTreeItemMetrics::IconSize())
+					.HeightOverride(FSceneOutlinerDefaultTreeItemMetrics::IconSize())
+					[
+						SNew(SImage)
+						.Image_Lambda([this]() -> const FSlateBrush* 
+						{
+							if (const TSharedPtr<FWorkspaceOutlinerTreeItem> SharedTreeItem = TreeItem.Pin())
+							{
+								if (SharedTreeItem->ItemDetails.IsValid())
+								{
+									return SharedTreeItem->ItemDetails->GetItemIcon();
+								}
+							}
+							
+							return FAppStyle::GetBrush(TEXT("ClassIcon.Default"));
+						})
+						.ColorAndOpacity(FSlateColor::UseForeground())
+					]
+				]
+				+ SHorizontalBox::Slot()
+				.FillWidth(1.0f)
+				.VAlign(VAlign_Center)
+				.Padding(0.0f, 2.0f)
+				[
+					SNew(STextBlock)
+					.Text(this, &SWorkspaceOutlinerTreelabel::GetDisplayText)
+					.HighlightText(SceneOutliner.GetFilterHighlightText())
+					.ColorAndOpacity(this, &SWorkspaceOutlinerTreelabel::GetForegroundColor)
+				]
+			];
+		}
+
+		FText GetDisplayText() const
+		{
+			if (const TSharedPtr<FWorkspaceOutlinerTreeItem> Item = TreeItem.Pin())
+			{				
+				return FText::FromString(Item->GetDisplayString());
+			}
+			return FText();
+		}
+		
+		virtual FSlateColor GetForegroundColor() const override
+		{
+			const TOptional<FLinearColor> BaseColor = FSceneOutlinerCommonLabelData::GetForegroundColor(*TreeItem.Pin());
+			return BaseColor.IsSet() ? BaseColor.GetValue() : FSlateColor::UseForeground();
+		}
+
+	private:
+		TWeakPtr<FWorkspaceOutlinerTreeItem> TreeItem;
+	};	
+	
+	FWorkspaceOutlinerTreeItem::FWorkspaceOutlinerTreeItem(const FItemData& InItemData) : ISceneOutlinerTreeItem(FWorkspaceOutlinerTreeItem::Type), Export(InItemData.Export)
+	{
+		ItemDetails = FWorkspaceEditorModule::GetOutlinerItemDetails(MakeOutlinerDetailsId(Export));
+	}
+
+	bool FWorkspaceOutlinerTreeItem::IsValid() const
+	{
+		return Export.Identifier.IsValid();
+	}
+
+	FSceneOutlinerTreeItemID FWorkspaceOutlinerTreeItem::GetID() const
+	{
+		return GetTypeHash(Export);
+	}
+
+	FString FWorkspaceOutlinerTreeItem::GetDisplayString() const
+	{
+		return Export.Identifier.ToString();
+	}
+
+	TSharedRef<SWidget> FWorkspaceOutlinerTreeItem::GenerateLabelWidget(ISceneOutliner& Outliner, const STableRow<FSceneOutlinerTreeItemPtr>& InRow)
+	{
+		return SNew(SWorkspaceOutlinerTreelabel, *this, Outliner, InRow);
+	}
+}
+
+#undef LOCTEXT_NAMESPACE // "WorkspaceOutlinerTreeItem"
\ No newline at end of file
diff --git a/Engine/Plugins/Experimental/Workspace/Source/WorkspaceEditor/Private/Outliner/WorkspaceOutlinerTreeItem.h b/Engine/Plugins/Experimental/Workspace/Source/WorkspaceEditor/Private/Outliner/WorkspaceOutlinerTreeItem.h
new file mode 100644
--- /dev/null
+++ b/Engine/Plugins/Experimental/Workspace/Source/WorkspaceEditor/Private/Outliner/WorkspaceOutlinerTreeItem.h
@@ -0,0 +1,35 @@
+// Copyright Epic Games, Inc. All Rights Reserved.
+
+#pragma once
+
+#include "ISceneOutlinerTreeItem.h"
+#include "WorkspaceAssetRegistryInfo.h"
+
+namespace UE::Workspace
+{
+	class IWorkspaceOutlinerItemDetails;
+	struct FWorkspaceOutlinerTreeItem : ISceneOutlinerTreeItem
+	{	
+		static const FSceneOutlinerTreeItemType Type;
+		struct FItemData
+		{
+			const FWorkspaceOutlinerItemExport& Export;
+		};
+				
+		FWorkspaceOutlinerTreeItem(const FItemData& InItemData);
+		virtual ~FWorkspaceOutlinerTreeItem() override {}
+
+		// Begin ISceneOutlinerTreeItem overrides
+		virtual bool IsValid() const override;
+		virtual FSceneOutlinerTreeItemID GetID() const override;
+		virtual FString GetDisplayString() const override;
+		virtual bool CanInteract() const override { return true; }
+		virtual TSharedRef<SWidget> GenerateLabelWidget(ISceneOutliner& Outliner, const STableRow<FSceneOutlinerTreeItemPtr>& InRow) override;
+		// End ISceneOutlinerTreeItem overrides
+
+		// AssetRegistry export this tree item represents in the Outliner 
+		FWorkspaceOutlinerItemExport Export;
+		// Cached _optional_ details provider for the FWorkspaceOutlinerItemExport its inner data
+		TSharedPtr<IWorkspaceOutlinerItemDetails> ItemDetails;
+	};
+}
diff --git a/Engine/Plugins/Experimental/Workspace/Source/WorkspaceEditor/Private/SWorkspaceTabWrapper.cpp b/Engine/Plugins/Experimental/Workspace/Source/WorkspaceEditor/Private/SWorkspaceTabWrapper.cpp
new file mode 100644
--- /dev/null
+++ b/Engine/Plugins/Experimental/Workspace/Source/WorkspaceEditor/Private/SWorkspaceTabWrapper.cpp
@@ -0,0 +1,200 @@
+// Copyright Epic Games, Inc. All Rights Reserved.
+
+#include "SWorkspaceTabWrapper.h"
+#include "Widgets/Layout/SSeparator.h"
+
+void SWorkspaceTabWrapper::Construct( const FArguments& InArgs, TSharedPtr<class FTabInfo> InTabInfo, TSharedPtr<UE::Workspace::FWorkspaceEditor> InWorkspaceEditor, UObject* InDocumentID)
+{
+	Content = InArgs._Content.Widget;
+	WeakWorkspaceEditor = InWorkspaceEditor;
+	WeakDocumentObject = InDocumentID;
+
+	// Set-up shared breadcrumb defaults JDB TODO figure out correct padding to align fake title with breadcrumbs
+    const FMargin BreadcrumbTrailPadding = FMargin(4.f, 2.f);
+    const FSlateBrush* BreadcrumbButtonImage = FAppStyle::GetBrush("BreadcrumbTrail.Delimiter");
+	
+	ChildSlot
+	[
+		SNew(SVerticalBox)
+		+SVerticalBox::Slot()
+		.AutoHeight()
+		[
+			SNew(SHorizontalBox)
+			+ SHorizontalBox::Slot()
+			.AutoWidth()
+			[
+				InTabInfo->CreateHistoryNavigationWidget()
+			]
+			+ SHorizontalBox::Slot()
+			.AutoWidth()
+			[
+				SNew(SSeparator)
+				.Orientation(Orient_Vertical)
+			]
+			// Title text/icon
+			+SHorizontalBox::Slot()
+			[
+				SNew(SHorizontalBox)
+				+SHorizontalBox::Slot()
+				.FillWidth(1.f)
+				[
+					SNew(SHorizontalBox)
+					+SHorizontalBox::Slot()
+					.AutoWidth()
+					//.Padding( 10.0f,5.0f )
+					.VAlign(VAlign_Center)
+					[
+						SNew(SImage)
+						.Image( this, &SWorkspaceTabWrapper::GetTabIcon )
+						.ColorAndOpacity(FSlateColor::UseForeground())
+					]
+					+SHorizontalBox::Slot()
+					.FillWidth(1.f)
+					[
+						SAssignNew(BreadcrumbTrailScrollBox, SScrollBox)
+						.Orientation(Orient_Horizontal)
+						.ScrollBarVisibility(EVisibility::Collapsed)
+
+						+SScrollBox::Slot()
+						.Padding(0.f)
+						.VAlign(VAlign_Center)
+						[
+							SNew(SHorizontalBox)
+							// show fake 'root' breadcrumb for the title
+							+SHorizontalBox::Slot()
+							.AutoWidth()
+							.VAlign(VAlign_Center)
+							.Padding(BreadcrumbTrailPadding)
+							[
+								SNew(SVerticalBox)
+								+SVerticalBox::Slot()
+								.FillHeight(1.f)
+								[
+									SNew(STextBlock)
+									.Text(this, &SWorkspaceTabWrapper::GetWorkspaceName)
+									.TextStyle( FAppStyle::Get(), TEXT("GraphBreadcrumbButtonText") )
+									.Visibility( this, &SWorkspaceTabWrapper::IsWorkspaceNameVisible )
+								]
+								
+							]
+							+SHorizontalBox::Slot()
+							.AutoWidth()
+							.VAlign(VAlign_Center)
+							.Padding(BreadcrumbTrailPadding)
+							[
+								SNew(SImage)
+								.Image( BreadcrumbButtonImage )
+								.Visibility( this, &SWorkspaceTabWrapper::IsWorkspaceNameVisible )
+								.ColorAndOpacity(FSlateColor::UseForeground())
+							]
+
+							// New style breadcrumb
+							+SHorizontalBox::Slot()
+							.AutoWidth()
+							.VAlign(VAlign_Center)
+							[
+								SAssignNew(BreadcrumbTrail, SBreadcrumbTrail<TSharedPtr<UE::Workspace::FWorkspaceBreadcrumb>>)
+								.ButtonStyle(FAppStyle::Get(), "GraphBreadcrumbButton")
+								.TextStyle(FAppStyle::Get(), "GraphBreadcrumbButtonText")
+								.ButtonContentPadding(BreadcrumbTrailPadding)
+								.DelimiterImage(BreadcrumbButtonImage)
+								.PersistentBreadcrumbs(false)
+								.OnCrumbClicked_Lambda([](const TSharedPtr<UE::Workspace::FWorkspaceBreadcrumb> InBreadcrumb){ InBreadcrumb->OnClicked.ExecuteIfBound(); })
+							]
+						]
+					]
+				]
+			]
+		]
+		+ SVerticalBox::Slot()
+		.AutoHeight()
+		[
+			SNew(SSeparator)
+			.Orientation(Orient_Horizontal)
+		]
+		+SVerticalBox::Slot()
+		.FillHeight(1.0f)
+		[
+			GetContent()
+		]
+	];
+
+	RebuildBreadcrumbTrail();
+} 
+
+void SWorkspaceTabWrapper::RebuildBreadcrumbTrail() const
+{
+	BreadcrumbTrail->ClearCrumbs(false);
+
+	TArray<TSharedPtr<UE::Workspace::FWorkspaceBreadcrumb>> Breadcrumbs;
+	if (UObject* DocumentID = WeakDocumentObject.Get())
+	{
+		if (const TSharedPtr<UE::Workspace::FWorkspaceEditor> SharedWorkspaceEditor = WeakWorkspaceEditor.Pin())
+		{
+			const UE::Workspace::FWorkspaceEditorModule& WorkspaceEditorModule = FModuleManager::LoadModuleChecked<UE::Workspace::FWorkspaceEditorModule>("WorkspaceEditor");
+			if(const UE::Workspace::FObjectDocumentArgs* DocumentArgs = WorkspaceEditorModule.FindObjectDocumentType(DocumentID->GetClass()->GetClassPathName()))
+			{
+				if(DocumentArgs->OnGetDocumentBreadcrumbTrail.IsBound())
+				{
+					DocumentArgs->OnGetDocumentBreadcrumbTrail.Execute(UE::Workspace::FWorkspaceEditorContext(SharedWorkspaceEditor.ToSharedRef(), DocumentID), Breadcrumbs);
+				}
+			}
+		}			
+	}
+
+	// Widgets have to be added in reverse order
+	for (int32 Index = Breadcrumbs.Num() - 1; Index >= 0; --Index)
+	{
+		TSharedPtr<UE::Workspace::FWorkspaceBreadcrumb> Breadcrumb = Breadcrumbs[Index];
+		BreadcrumbTrail->PushCrumb(	TAttribute<FText>::CreateLambda([Breadcrumb]() { return Breadcrumb->OnGetLabel.IsBound() ? Breadcrumb->OnGetLabel.Execute().Get() : FText::GetEmpty(); }), Breadcrumb);
+	}
+}
+
+const FSlateBrush* SWorkspaceTabWrapper::GetTabIcon() const
+{
+	if (UObject* DocumentID = WeakDocumentObject.Get())
+	{
+		if (const TSharedPtr<UE::Workspace::FWorkspaceEditor> SharedWorkspaceEditor = WeakWorkspaceEditor.Pin())
+		{
+			const UE::Workspace::FWorkspaceEditorModule& WorkspaceEditorModule = FModuleManager::LoadModuleChecked<UE::Workspace::FWorkspaceEditorModule>("WorkspaceEditor");
+			if(const UE::Workspace::FObjectDocumentArgs* DocumentArgs = WorkspaceEditorModule.FindObjectDocumentType(DocumentID->GetClass()->GetClassPathName()))
+			{
+				if(DocumentArgs->OnGetTabIcon.IsBound())
+				{
+					return DocumentArgs->OnGetTabIcon.Execute(UE::Workspace::FWorkspaceEditorContext(SharedWorkspaceEditor.ToSharedRef(), DocumentID));
+				}
+			}
+		}			
+	}
+	
+	return nullptr;
+}
+
+EVisibility SWorkspaceTabWrapper::IsWorkspaceNameVisible() const
+{
+	if (const TSharedPtr<UE::Workspace::FWorkspaceEditor> SharedWorkspaceEditor = WeakWorkspaceEditor.Pin())
+	{
+		if (SharedWorkspaceEditor->Workspace)
+		{
+			return EVisibility::Visible;
+		}
+	}
+
+	return EVisibility::Collapsed;
+}
+
+FText  SWorkspaceTabWrapper::GetWorkspaceName() const
+{
+	if (const TSharedPtr<UE::Workspace::FWorkspaceEditor> SharedWorkspaceEditor = WeakWorkspaceEditor.Pin())
+	{
+		if (SharedWorkspaceEditor->Workspace)
+		{
+			return FText::FromName(SharedWorkspaceEditor->Workspace->GetFName());
+		}
+	}
+
+	return FText::GetEmpty();
+}
+
+	
+
diff --git a/Engine/Plugins/Experimental/Workspace/Source/WorkspaceEditor/Private/SWorkspaceTabWrapper.h b/Engine/Plugins/Experimental/Workspace/Source/WorkspaceEditor/Private/SWorkspaceTabWrapper.h
new file mode 100644
--- /dev/null
+++ b/Engine/Plugins/Experimental/Workspace/Source/WorkspaceEditor/Private/SWorkspaceTabWrapper.h
@@ -0,0 +1,36 @@
+// Copyright Epic Games, Inc. All Rights Reserved.
+
+#pragma once
+
+#include "Workspace.h"
+#include "Widgets/SCompoundWidget.h"
+#include "Widgets/Navigation/SBreadcrumbTrail.h"
+#include "WorkflowOrientedApp/WorkflowTabManager.h"
+#include "WorkspaceEditor.h"
+#include "WorkspaceEditorModule.h"
+
+class SWorkspaceTabWrapper : public SCompoundWidget
+{
+public:
+	SWorkspaceTabWrapper() = default;
+
+	SLATE_BEGIN_ARGS( SWorkspaceTabWrapper ){}
+		SLATE_DEFAULT_SLOT(FArguments, Content)
+	SLATE_END_ARGS()
+	
+	void Construct(const FArguments& InArgs, TSharedPtr<class FTabInfo> InTabInfo, TSharedPtr<UE::Workspace::FWorkspaceEditor> InWorkspaceEditor, UObject* InDocumentID);
+	TSharedRef<SWidget> GetContent() const { return Content.ToSharedRef(); }
+	
+protected:
+	void RebuildBreadcrumbTrail() const;
+	const FSlateBrush* GetTabIcon() const;
+	EVisibility IsWorkspaceNameVisible() const;
+	FText GetWorkspaceName() const;
+
+	TSharedPtr<SWidget> Content = nullptr;
+	TWeakPtr<UE::Workspace::FWorkspaceEditor> WeakWorkspaceEditor = nullptr;
+	TWeakObjectPtr<UObject> WeakDocumentObject = nullptr;
+
+	TSharedPtr<SScrollBox> BreadcrumbTrailScrollBox = nullptr;
+	TSharedPtr<SBreadcrumbTrail<TSharedPtr<UE::Workspace::FWorkspaceBreadcrumb>>> BreadcrumbTrail = nullptr;	
+};
diff --git a/Engine/Plugins/Experimental/Workspace/Source/WorkspaceEditor/Private/SWorkspaceView.cpp b/Engine/Plugins/Experimental/Workspace/Source/WorkspaceEditor/Private/SWorkspaceView.cpp
--- a/Engine/Plugins/Experimental/Workspace/Source/WorkspaceEditor/Private/SWorkspaceView.cpp
+++ b/Engine/Plugins/Experimental/Workspace/Source/WorkspaceEditor/Private/SWorkspaceView.cpp
@@ -5,48 +5,32 @@
 #include "ContentBrowserModule.h"
 #include "IContentBrowserSingleton.h"
 #include "SAssetDropTarget.h"
+#include "SceneOutlinerPublicTypes.h"
 #include "ScopedTransaction.h"
+#include "SSceneOutliner.h"
 #include "WorkspaceSchema.h"
 #include "Framework/Commands/GenericCommands.h"
+#include "Outliner/WorkspaceOutlinerColumns.h"
+#include "Outliner/WorkspaceOutlinerMode.h"
 
 #define LOCTEXT_NAMESPACE "SWorkspaceView"
 
 namespace UE::Workspace
 {
 
-void SWorkspaceView::Construct(const FArguments& InArgs, UWorkspace* InWorkspace)
+void SWorkspaceView::Construct(const FArguments& InArgs, UWorkspace* InWorkspace, TSharedRef<UE::Workspace::IWorkspaceEditor> InWorkspaceEditor)
 {
 	Workspace = InWorkspace;
-	OnAssetsOpened = InArgs._OnAssetsOpened;
 
-	UICommandList = MakeShared<FUICommandList>();
-
-	UICommandList->MapAction(FGenericCommands::Get().Delete,
-		FExecuteAction::CreateSP(this, &SWorkspaceView::HandleDelete),
-		FCanExecuteAction::CreateSP(this, &SWorkspaceView::HasValidSelection));
-
-	Workspace->ModifiedDelegate.AddSP(this, &SWorkspaceView::HandleWorkspaceModified);
-
-	FContentBrowserModule& ContentBrowserModule = FModuleManager::Get().LoadModuleChecked<FContentBrowserModule>(TEXT("ContentBrowser"));
-
-	FAssetPickerConfig AssetPickerConfig;
-	AssetPickerConfig.InitialAssetViewType = EAssetViewType::List;
-	AssetPickerConfig.OnAssetsActivated = FOnAssetsActivated::CreateLambda([this](const TArray<FAssetData>& InAssets, EAssetTypeActivationMethod::Type InActivationMethod)
+	FSceneOutlinerInitializationOptions InitOptions;
 	{
-		if( InActivationMethod == EAssetTypeActivationMethod::DoubleClicked ||
-			InActivationMethod == EAssetTypeActivationMethod::Opened)
-		{
-			OnAssetsOpened.ExecuteIfBound(InAssets);
-		}
-	});
-	AssetPickerConfig.RefreshAssetViewDelegates.Add(&RefreshAssetViewDelegate);
-	AssetPickerConfig.GetCurrentSelectionDelegates.Add(&GetCurrentSelectionDelegate);
-	AssetPickerConfig.Filter = MakeARFilter();
-	AssetPickerConfig.OnShouldFilterAsset = FOnShouldFilterAsset::CreateLambda([this](const FAssetData& InAsset)
-	{
-		return !Workspace->AssetEntries.ContainsByPredicate([InAsset](const UWorkspaceAssetEntry* Entry)-> bool { return Entry->Asset == TSoftObjectPtr<UObject>(InAsset.GetSoftObjectPath()); });
-	});
-	AssetPickerConfig.AssetShowWarningText = LOCTEXT("EmptyWorkspaceText", "Workspace is empty.\nDrag-drop to add assets to this workspace.");
+		InitOptions.OutlinerIdentifier = TEXT("WorkspaceEditorOutliner");
+		InitOptions.bShowHeaderRow = true;
+		InitOptions.ColumnMap.Add(FSceneOutlinerBuiltInColumnTypes::Label(), FSceneOutlinerColumnInfo(ESceneOutlinerColumnVisibility::Visible, 0));
+		InitOptions.ColumnMap.Add(FWorkspaceOutlinerSourceControlColumn::GetID(), FSceneOutlinerColumnInfo(ESceneOutlinerColumnVisibility::Visible, 100, FCreateSceneOutlinerColumn::CreateLambda([](ISceneOutliner& InSceneOutliner) { return MakeShareable(new FWorkspaceOutlinerSourceControlColumn(InSceneOutliner)); }), false));	
+		InitOptions.ModeFactory = FCreateSceneOutlinerMode::CreateLambda([this, WeakWorkspaceEditor=InWorkspaceEditor.ToWeakPtr()](SSceneOutliner* InOutliner) { return new UE::Workspace::FWorkspaceOutlinerMode(UE::Workspace::FWorkspaceOutlinerMode(InOutliner, Workspace, WeakWorkspaceEditor)); });
+	}
+	SceneWorkspaceOutliner = SNew(SSceneOutliner, InitOptions);
 
 	ChildSlot
 	[
@@ -73,61 +57,11 @@ void SWorkspaceView::Construct(const FArguments& InArgs, UWorkspace* InWorkspace
 		})
 		.Content()
 		[
-			ContentBrowserModule.Get().CreateAssetPicker(AssetPickerConfig)
+			SceneWorkspaceOutliner.ToSharedRef()
 		]
 	];
 }
 
-FARFilter SWorkspaceView::MakeARFilter()
-{
-	FARFilter Filter;
-	Filter.ClassPaths = Workspace->GetSchema()->GetSupportedAssetClassPaths();
-	Filter.bRecursiveClasses = true;
-
-	return Filter;
-}
-
-FReply SWorkspaceView::OnKeyDown(const FGeometry& MyGeometry, const FKeyEvent& InKeyEvent)
-{
-	if (UICommandList.IsValid() && UICommandList->ProcessCommandBindings(InKeyEvent))
-	{
-		return FReply::Handled();
-	}
-	return FReply::Unhandled();
-}
-
-void SWorkspaceView::HandleDelete()
-{
-	if(GetCurrentSelectionDelegate.IsBound())
-	{
-		TArray<FAssetData> Selection = GetCurrentSelectionDelegate.Execute();
-
-		if(Selection.Num() > 0)
-		{
-			FScopedTransaction Transaction(LOCTEXT("RemoveAssets", "Remove assets from workspace"));
-
-			Workspace->RemoveAssets(Selection);
-		}
-	}
-}
-
-bool SWorkspaceView::HasValidSelection() const 
-{
-	if(GetCurrentSelectionDelegate.IsBound())
-	{
-		TArray<FAssetData> Selection = GetCurrentSelectionDelegate.Execute();
-		return Selection.Num() > 0;
-	}
-	return false;
-}
-
-void SWorkspaceView::HandleWorkspaceModified(UWorkspace* InWorkspace)
-{
-	check(InWorkspace == Workspace);
-
-	RefreshAssetViewDelegate.ExecuteIfBound(true);
-}
-
 }
 
 #undef LOCTEXT_NAMESPACE
\ No newline at end of file
diff --git a/Engine/Plugins/Experimental/Workspace/Source/WorkspaceEditor/Private/SWorkspaceView.h b/Engine/Plugins/Experimental/Workspace/Source/WorkspaceEditor/Private/SWorkspaceView.h
--- a/Engine/Plugins/Experimental/Workspace/Source/WorkspaceEditor/Private/SWorkspaceView.h
+++ b/Engine/Plugins/Experimental/Workspace/Source/WorkspaceEditor/Private/SWorkspaceView.h
@@ -6,43 +6,22 @@
 #include "Widgets/SCompoundWidget.h"
 
 class UWorkspace;
+class UWorkSpaceAssetUserData;
+class SSceneOutliner;
 
 namespace UE::Workspace
 {
+class IWorkspaceEditor;
 
 class SWorkspaceView : public SCompoundWidget
 {
-	using FOnAssetsOpened = TDelegate<void(TConstArrayView<FAssetData>)>; 
-
 	SLATE_BEGIN_ARGS(SWorkspaceView) {}
-
-	SLATE_EVENT(FOnAssetsOpened, OnAssetsOpened)
-
 	SLATE_END_ARGS()
 
-	void Construct(const FArguments& InArgs, UWorkspace* InWorkspace);
-
+	void Construct(const FArguments& InArgs, UWorkspace* InWorkspace, TSharedRef<UE::Workspace::IWorkspaceEditor> InWorkspaceEditor);
 private:
-	// SWidget interface
-	virtual FReply OnKeyDown(const FGeometry& MyGeometry, const FKeyEvent& InKeyEvent) override;
-
-	void HandleDelete();
-
-	bool HasValidSelection() const;
-
-	void HandleWorkspaceModified(UWorkspace* InWorkspace);
-
-	FARFilter MakeARFilter();
-
 	UWorkspace* Workspace = nullptr;
-
-	FOnAssetsOpened OnAssetsOpened;
-
-	FRefreshAssetViewDelegate RefreshAssetViewDelegate;
-
-	FGetCurrentSelectionDelegate GetCurrentSelectionDelegate;
-
-	TSharedPtr<FUICommandList> UICommandList;
+	TSharedPtr<SSceneOutliner> SceneWorkspaceOutliner;
 };
 
 };
\ No newline at end of file
diff --git a/Engine/Plugins/Experimental/Workspace/Source/WorkspaceEditor/Private/Workspace.cpp b/Engine/Plugins/Experimental/Workspace/Source/WorkspaceEditor/Private/Workspace.cpp
--- a/Engine/Plugins/Experimental/Workspace/Source/WorkspaceEditor/Private/Workspace.cpp
+++ b/Engine/Plugins/Experimental/Workspace/Source/WorkspaceEditor/Private/Workspace.cpp
@@ -9,6 +9,9 @@
 #include "UObject/AssetRegistryTagsContext.h"
 #include "UObject/ConstructorHelpers.h"
 #include "UObject/FortniteMainBranchObjectVersion.h"
+#include "Engine/ExternalAssetDependencyGatherer.h"
+
+REGISTER_ASSETDEPENDENCY_GATHERER(FExternalAssetDependencyGatherer, UWorkspace);
 
 const FName UWorkspace::ExportsAssetRegistryTag = TEXT("Exports");
 
@@ -246,6 +249,28 @@ void UWorkspace::BroadcastModified()
 	}
 }
 
+void UWorkspace::GetAssets(TArray<TObjectPtr<UObject>>& OutAssets) const
+{
+	for(const UWorkspaceAssetEntry* AssetEntry : AssetEntries)
+	{
+		if (AssetEntry)
+		{
+			if (UObject* Asset = AssetEntry->Asset.Get())
+			{
+				OutAssets.Add(Asset);
+			}
+		}	
+	}
+}
+
+void UWorkspace::GetAssetDataEntries(TArray<FAssetData>& OutAssetDataEntries) const
+{
+	FARFilter Filter;
+	Algo::Transform(AssetEntries, Filter.SoftObjectPaths, [](const UWorkspaceAssetEntry* AssetEntry) { return AssetEntry->Asset.ToSoftObjectPath(); });	
+	const IAssetRegistry& AssetRegistry = FAssetRegistryModule::GetRegistry();
+	AssetRegistry.GetAssets(Filter, OutAssetDataEntries);
+}
+
 void UWorkspace::ReportError(const TCHAR* InMessage) const
 {
 #if WITH_EDITOR
@@ -272,6 +297,25 @@ void UWorkspace::Serialize(FArchive& Ar)
 {
 	UObject::Serialize(Ar);
 	Ar.UsingCustomVersion(FFortniteMainBranchObjectVersion::GUID);
+
+	const bool bIsDuplicating = (Ar.GetPortFlags() & PPF_Duplicate) != 0;
+	if (bIsDuplicating)
+	{
+		Ar << AssetEntries;
+		Ar << State;
+	}
+}
+
+bool UWorkspace::Rename(const TCHAR* NewName, UObject* NewOuter, ERenameFlags Flags)
+{
+	FExternalPackageHelper::FRenameExternalObjectsHelperContext Context(this, Flags);
+	return Super::Rename(NewName, NewOuter, Flags);
+}
+
+void UWorkspace::PreDuplicate(FObjectDuplicationParameters& DupParams)
+{
+	UObject::PreDuplicate(DupParams);	
+	FExternalPackageHelper::DuplicateExternalPackages(this, DupParams);
 }
 
 void UWorkspace::PostLoad()
diff --git a/Engine/Plugins/Experimental/Workspace/Source/WorkspaceEditor/Private/Workspace.h b/Engine/Plugins/Experimental/Workspace/Source/WorkspaceEditor/Private/Workspace.h
--- a/Engine/Plugins/Experimental/Workspace/Source/WorkspaceEditor/Private/Workspace.h
+++ b/Engine/Plugins/Experimental/Workspace/Source/WorkspaceEditor/Private/Workspace.h
@@ -20,6 +20,10 @@ namespace UE::Workspace
 	class FWorkspaceEditor;
 	class FWorkspaceEditorModule;
 	class SWorkspaceView;
+	class FWorkspaceOutlinerHierarchy;
+	class FWorkspaceOutlinerMode;
+
+	struct FWorkspaceOutliner;
 }
 
 namespace UE::Workspace
@@ -62,6 +66,10 @@ class UWorkspace : public UObject
 	friend class UE::Workspace::FWorkspaceEditor;
 	friend class UE::Workspace::FWorkspaceEditorModule;
 	friend class UE::Workspace::SWorkspaceView;
+	friend struct UE::Workspace::FWorkspaceOutliner;
+	friend class UE::Workspace::FWorkspaceOutlinerHierarchy;
+	friend class UE::Workspace::FWorkspaceOutlinerMode;
+	
 	friend class UWorkspaceFactory;
 	friend class UAssetDefinition_Workspace;
 	friend class UWorkspaceState;
@@ -104,6 +112,11 @@ class UWorkspace : public UObject
 
 	void BroadcastModified();
 
+	// Returns all contained Asset objects
+	void GetAssets(TArray<TObjectPtr<UObject>>& OutAssets) const;
+	// Returns all contained AssetData entries
+	void GetAssetDataEntries(TArray<FAssetData>& OutAssetDataEntries) const;
+
 	// UObject interface
 	virtual void PostLoad() override;
 	virtual void PostDuplicate(bool bDuplicateForPIE) override;
@@ -111,6 +124,8 @@ class UWorkspace : public UObject
 	virtual void PostTransacted(const FTransactionObjectEvent& TransactionEvent) override;
 	virtual bool IsEditorOnly() const override { return true; }
 	virtual void Serialize(FArchive& Ar) override;
+	virtual bool Rename(const TCHAR* NewName = nullptr, UObject* NewOuter = nullptr, ERenameFlags Flags = REN_None) override;	
+	virtual void PreDuplicate(FObjectDuplicationParameters& DupParams) override;
 
 	void PostLoadExternalPackages();
 	
diff --git a/Engine/Plugins/Experimental/Workspace/Source/WorkspaceEditor/Private/WorkspaceEditor.cpp b/Engine/Plugins/Experimental/Workspace/Source/WorkspaceEditor/Private/WorkspaceEditor.cpp
--- a/Engine/Plugins/Experimental/Workspace/Source/WorkspaceEditor/Private/WorkspaceEditor.cpp
+++ b/Engine/Plugins/Experimental/Workspace/Source/WorkspaceEditor/Private/WorkspaceEditor.cpp
@@ -1,10 +1,10 @@
 // Copyright Epic Games, Inc. All Rights Reserved.
 
 #include "WorkspaceEditor.h"
-
 #include "Workspace.h"
 #include "AssetDocumentSummoner.h"
 #include "ExternalPackageHelper.h"
+#include "IStructureDetailsView.h"
 #include "Widgets/Docking/SDockTab.h"
 #include "WorkflowOrientedApp/WorkflowUObjectDocuments.h"
 #include "WorkspaceAssetEditor.h"
@@ -12,6 +12,9 @@
 #include "WorkspaceDocumentState.h"
 #include "WorkspaceEditorModule.h"
 #include "SWorkspaceView.h"
+#include "Toolkits/AssetEditorToolkitMenuContext.h"
+#include "WorkspaceEditorCommands.h"
+#include "SWorkspaceTabWrapper.h"
 
 #define LOCTEXT_NAMESPACE "WorkspaceEditor"
 
@@ -53,7 +56,8 @@ void FWorkspaceEditor::CreateWidgets()
 	LeftAssetDocumentSummoner->SetAllowedClassPaths(WorkspaceEditorModule.GetAllowedObjectTypesForArea(WorkspaceTabs::LeftDocumentArea));
 	DocumentManager->RegisterDocumentFactory(LeftAssetDocumentSummoner);
 
-	const TSharedRef<FAssetDocumentSummoner> MiddleAssetDocumentSummoner = MakeShared<FAssetDocumentSummoner>(WorkspaceTabs::MiddleDocumentArea, SharedThis(this));
+	constexpr bool bAllowUnsupportedClasses = true;
+	const TSharedRef<FAssetDocumentSummoner> MiddleAssetDocumentSummoner = MakeShared<FAssetDocumentSummoner>(WorkspaceTabs::MiddleDocumentArea, SharedThis(this), bAllowUnsupportedClasses);
 	MiddleAssetDocumentSummoner->SetAllowedClassPaths(WorkspaceEditorModule.GetAllowedObjectTypesForArea(WorkspaceTabs::MiddleDocumentArea));
 	DocumentManager->RegisterDocumentFactory(MiddleAssetDocumentSummoner);
 
@@ -100,17 +104,7 @@ void FWorkspaceEditor::CreateWidgets()
 		)
 	);
 
-	TWeakPtr<FWorkspaceEditor> WeakToolkit = StaticCastSharedRef<FWorkspaceEditor>(AsShared());
-	WorkspaceView = SNew(SWorkspaceView, Workspace)
-		.OnAssetsOpened_Lambda([WeakToolkit](TConstArrayView<FAssetData> InAssets)
-		{
-			if(const TSharedPtr<FWorkspaceEditor> HostingApp = WeakToolkit.Pin())
-			{
-				HostingApp->OpenAssets(InAssets);
-			}
-		});
-
-	
+	WorkspaceView = SNew(SWorkspaceView, Workspace, StaticCastSharedRef<UE::Workspace::IWorkspaceEditor>(AsShared()));
 	BindCommands();
 }
 
@@ -133,11 +127,12 @@ void FWorkspaceEditor::RestoreEditedObjectState()
 		{
 			if(const TSharedPtr<SDockTab> DockTab = OpenDocument(Object, FDocumentTracker::RestorePreviousDocument))
 			{
+				const TSharedRef<SWorkspaceTabWrapper> TabWrapper = StaticCastSharedRef<SWorkspaceTabWrapper>(DockTab->GetContent());	
 				FWorkspaceEditorModule& WorkspaceEditorModule = FModuleManager::LoadModuleChecked<FWorkspaceEditorModule>("WorkspaceEditor");
 				const FObjectDocumentArgs* DocumentArgs = WorkspaceEditorModule.FindObjectDocumentType(Object->GetClass()->GetClassPathName());
 				if(DocumentArgs != nullptr && DocumentArgs->OnSetDocumentState.IsBound())
 				{
-					DocumentArgs->OnSetDocumentState.Execute(FWorkspaceEditorContext(SharedThis(this), Object), DockTab->GetContent(), DocumentState);
+					DocumentArgs->OnSetDocumentState.Execute(FWorkspaceEditorContext(SharedThis(this), Object), TabWrapper->GetContent(), DocumentState);
 				}
 			}
 		}
@@ -168,7 +163,8 @@ TSharedPtr<SDockTab> FWorkspaceEditor::OpenDocument(const UObject* InForObject,
 		const FObjectDocumentArgs* DocumentArgs = WorkspaceEditorModule.FindObjectDocumentType(InForObject->GetClass()->GetClassPathName());
 		if(DocumentArgs != nullptr && DocumentArgs->OnGetDocumentState.IsBound())
 		{
-			RecordDocumentState(DocumentArgs->OnGetDocumentState.Execute(FWorkspaceEditorContext(SharedThis(this), const_cast<UObject*>(InForObject)), NewTab->GetContent()));
+			const TSharedRef<SWorkspaceTabWrapper> TabWrapper = StaticCastSharedRef<SWorkspaceTabWrapper>(NewTab->GetContent());
+			RecordDocumentState(DocumentArgs->OnGetDocumentState.Execute(FWorkspaceEditorContext(SharedThis(this), const_cast<UObject*>(InForObject)), TabWrapper->GetContent()));
 		}
 		else
 		{
@@ -185,7 +181,7 @@ void FWorkspaceEditor::OpenAssets(TConstArrayView<FAssetData> InAssets)
 	{
 		if(const UObject* LoadedAsset = Asset.GetAsset())
 		{
-			OpenDocument(LoadedAsset, FDocumentTracker::EOpenDocumentCause::OpenNewDocument);
+			OpenDocument(LoadedAsset, FDocumentTracker::EOpenDocumentCause::NavigatingCurrentDocument);
 		}
 	}
 }
@@ -194,7 +190,7 @@ void FWorkspaceEditor::OpenObjects(TConstArrayView<UObject*> InObjects)
 {
 	for(const UObject* Object : InObjects)
 	{
-		OpenDocument(Object, FDocumentTracker::EOpenDocumentCause::OpenNewDocument);
+		OpenDocument(Object, FDocumentTracker::EOpenDocumentCause::NavigatingCurrentDocument);
 	}
 }
 
@@ -226,7 +222,21 @@ void FWorkspaceEditor::RefreshDetails()
 }
 
 void FWorkspaceEditor::BindCommands()
-{
+{	
+	FWorkspaceAssetEditorCommands::Register();
+	const FWorkspaceAssetEditorCommands& Commands = FWorkspaceAssetEditorCommands::Get();
+
+	ToolkitCommands->MapAction
+	( 
+		Commands.NavigateBackward,
+		FExecuteAction::CreateRaw( this, &FWorkspaceEditor::NavigateBack )
+	);
+	
+	ToolkitCommands->MapAction
+	( 
+		Commands.NavigateForward,
+		FExecuteAction::CreateRaw( this, &FWorkspaceEditor::NavigateForward )
+	);
 }
 
 void FWorkspaceEditor::ExtendMenu()
@@ -287,6 +297,11 @@ FLinearColor FWorkspaceEditor::GetWorldCentricTabColorScale() const
 
 void FWorkspaceEditor::InitToolMenuContext(FToolMenuContext& InMenuContext)
 {
+	UAssetEditorToolkitMenuContext* ToolkitMenuContext = NewObject<UAssetEditorToolkitMenuContext>();
+	ToolkitMenuContext->Toolkit = AsShared();
+	InMenuContext.AddObject(ToolkitMenuContext);
+	
+	IWorkspaceEditor::InitToolMenuContext(InMenuContext);
 }
 
 void FWorkspaceEditor::SaveAsset_Execute()
@@ -353,6 +368,18 @@ void FWorkspaceEditor::RecordDocumentState(const TInstancedStruct<FWorkspaceDocu
 	State->DocumentStates.AddUnique(InState);
 }
 
+void FWorkspaceEditor::NavigateBack()
+{
+	const TSharedRef<FTabPayload_UObject> Payload = FTabPayload_UObject::Make(nullptr);
+	TSharedPtr<SDockTab> OpenedTab = DocumentManager->OpenDocument(Payload, FDocumentTracker::NavigateBackwards);
+}
+
+void FWorkspaceEditor::NavigateForward()
+{
+	const TSharedRef<FTabPayload_UObject> Payload = FTabPayload_UObject::Make(nullptr);
+	TSharedPtr<SDockTab> OpenedTab = DocumentManager->OpenDocument(Payload, FDocumentTracker::NavigateForwards);
+}
+
 bool FWorkspaceEditor::OnRequestClose(EAssetEditorCloseReason InCloseReason)
 {
 	TGuardValue<bool> ClosingDown(bClosingDown, true);
diff --git a/Engine/Plugins/Experimental/Workspace/Source/WorkspaceEditor/Private/WorkspaceEditor.h b/Engine/Plugins/Experimental/Workspace/Source/WorkspaceEditor/Private/WorkspaceEditor.h
--- a/Engine/Plugins/Experimental/Workspace/Source/WorkspaceEditor/Private/WorkspaceEditor.h
+++ b/Engine/Plugins/Experimental/Workspace/Source/WorkspaceEditor/Private/WorkspaceEditor.h
@@ -15,6 +15,8 @@ class FDocumentTabFactory;
 class FTabInfo;
 class FTabManager;
 class IToolkitHost;
+class SWorkspaceTabWrapper;
+class IStructureDetailsView;
 
 namespace UE::Workspace
 {
@@ -44,11 +46,12 @@ public:
 
 private:
 	friend class FWorkspaceEditorMode;
-	friend struct FGraphDocumentSummoner;
+	friend class FWorkspaceEditorModule;
 	friend class SGraphDocument;
+	friend SWorkspaceTabWrapper;
+	friend struct FGraphDocumentSummoner;
 	friend struct FWorkspaceTabSummoner;
 	friend struct FAssetDocumentSummoner;
-	friend class FWorkspaceEditorModule;
 
 	// FBaseAssetToolkit interface
 	virtual void RegisterTabSpawners(const TSharedRef<class FTabManager>& TabManager) override;
@@ -104,6 +107,9 @@ private:
 
 	void RecordDocumentState(const TInstancedStruct<FWorkspaceDocumentState>& InState) const;
 
+	void NavigateBack();
+	void NavigateForward();
+
 	/** The asset being edited */
 	TObjectPtr<UWorkspace> Workspace = nullptr;
 
@@ -116,4 +122,4 @@ private:
 	bool bClosingDown = false;
 };
 
-}
\ No newline at end of file
+}
diff --git a/Engine/Plugins/Experimental/Workspace/Source/WorkspaceEditor/Private/WorkspaceEditorCommands.cpp b/Engine/Plugins/Experimental/Workspace/Source/WorkspaceEditor/Private/WorkspaceEditorCommands.cpp
new file mode 100644
--- /dev/null
+++ b/Engine/Plugins/Experimental/Workspace/Source/WorkspaceEditor/Private/WorkspaceEditorCommands.cpp
@@ -0,0 +1,35 @@
+// Copyright Epic Games, Inc. All Rights Reserved.
+
+#include "WorkspaceEditorCommands.h"
+
+#include "Framework/Commands/Commands.h"
+#include "Framework/Commands/InputChord.h"
+#include "Framework/Commands/UICommandInfo.h"
+#include "Styling/AppStyle.h"
+
+#define LOCTEXT_NAMESPACE "WorkspaceEditorCommands"
+
+namespace UE::Workspace
+{
+	FWorkspaceAssetEditorCommands::FWorkspaceAssetEditorCommands()
+		: TCommands<FWorkspaceAssetEditorCommands>(
+				"WorkspaceAssetEditor",
+				NSLOCTEXT("Contexts", "WorkspaceAssetEditor", "Workspace Editor"),
+				NAME_None,
+				FAppStyle::GetAppStyleSetName()
+			)
+	{
+	}
+
+	void FWorkspaceAssetEditorCommands::RegisterCommands()
+	{
+		UI_COMMAND(NavigateBackward, "Navigate Backward", "Moves backwards to previous location",
+			EUserInterfaceActionType::Button, FInputChord(EKeys::ThumbMouseButton), FInputChord(EKeys::Left, EModifierKey::Alt));
+		UI_COMMAND(NavigateForward, "Navigate Forward", "Moves forwards to previous location",
+			EUserInterfaceActionType::Button, FInputChord(EKeys::ThumbMouseButton2), FInputChord(EKeys::Right, EModifierKey::Alt));
+	}
+
+}  // namespace UE::Workspace
+
+#undef LOCTEXT_NAMESPACE
+
diff --git a/Engine/Plugins/Experimental/Workspace/Source/WorkspaceEditor/Private/WorkspaceEditorCommands.h b/Engine/Plugins/Experimental/Workspace/Source/WorkspaceEditor/Private/WorkspaceEditorCommands.h
new file mode 100644
--- /dev/null
+++ b/Engine/Plugins/Experimental/Workspace/Source/WorkspaceEditor/Private/WorkspaceEditorCommands.h
@@ -0,0 +1,22 @@
+// Copyright Epic Games, Inc. All Rights Reserved.
+
+#pragma once
+
+#include "Framework/Commands/Commands.h"
+
+class FUICommandInfo;
+
+namespace UE::Workspace
+{
+	class FWorkspaceAssetEditorCommands : public TCommands<FWorkspaceAssetEditorCommands>
+	{
+	public:
+		FWorkspaceAssetEditorCommands();
+		virtual void RegisterCommands() override;	
+
+		TSharedPtr<FUICommandInfo> NavigateBackward;		
+		TSharedPtr<FUICommandInfo> NavigateForward;
+	};
+
+}  // namespace UE::Workspace
+
diff --git a/Engine/Plugins/Experimental/Workspace/Source/WorkspaceEditor/Private/WorkspaceEditorModule.cpp b/Engine/Plugins/Experimental/Workspace/Source/WorkspaceEditor/Private/WorkspaceEditorModule.cpp
--- a/Engine/Plugins/Experimental/Workspace/Source/WorkspaceEditor/Private/WorkspaceEditorModule.cpp
+++ b/Engine/Plugins/Experimental/Workspace/Source/WorkspaceEditor/Private/WorkspaceEditorModule.cpp
@@ -10,6 +10,7 @@
 #include "Workspace.h"
 #include "WorkspaceAssetEditor.h"
 #include "WorkspaceEditor.h"
+#include "WorkspaceEditorCommands.h"
 #include "WorkspaceFactory.h"
 #include "AssetRegistry/AssetRegistryModule.h"
 #include "EdGraph/EdGraph.h"
@@ -22,6 +23,8 @@
 namespace UE::Workspace
 {
 
+TMap<FOutlinerItemDetailsId, TSharedPtr<IWorkspaceOutlinerItemDetails>> FWorkspaceEditorModule::OutlinerItemDetails;
+
 void FWorkspaceEditorModule::RegisterObjectDocumentType(const FTopLevelAssetPath& InClassPath, const FObjectDocumentArgs& InParams)
 {
 	ensure(InParams.SpawnLocation != NAME_None);
@@ -64,7 +67,11 @@ FObjectDocumentArgs FWorkspaceEditorModule::CreateGraphDocumentArgs(const FGraph
 			.OnCanPasteNodes(InArgs.OnCanPasteNodes)
 			.OnPasteNodes(InArgs.OnPasteNodes)
 			.OnCanDuplicateSelectedNodes(InArgs.OnCanDuplicateSelectedNodes)
-			.OnDuplicateSelectedNodes(InArgs.OnDuplicateSelectedNodes);
+			.OnDuplicateSelectedNodes(InArgs.OnDuplicateSelectedNodes)
+			.OnNavigateHistoryForward_Lambda([WorkspaceEditor=StaticCastSharedRef<FWorkspaceEditor>(InContext.WorkspaceEditor)](){ WorkspaceEditor->NavigateForward(); })
+			.OnNavigateHistoryBack_Lambda([WorkspaceEditor=StaticCastSharedRef<FWorkspaceEditor>(InContext.WorkspaceEditor)](){ WorkspaceEditor->NavigateBack(); })
+		
+		;
 	});
 	Args.OnGetTabIcon = FOnGetTabIcon::CreateLambda([](const FWorkspaceEditorContext& InContext)
 	{
@@ -92,18 +99,24 @@ FObjectDocumentArgs FWorkspaceEditorModule::CreateGraphDocumentArgs(const FGraph
 	});
 	Args.OnGetDocumentState = FOnGetDocumentState::CreateLambda([](const FWorkspaceEditorContext& InContext, TSharedRef<SWidget> InWidget)
 	{
-		TSharedRef<SGraphDocument> GraphDocument = StaticCastSharedRef<SGraphDocument>(InWidget);
-		FVector2D ViewLocation;
-		float ZoomAmount;
-		GraphDocument->GraphEditor->GetViewLocation(ViewLocation, ZoomAmount);
+		FVector2D ViewLocation = FVector2d::ZeroVector;
+		float ZoomAmount = 0.f;
+		
+		if(const TSharedPtr<SGraphDocument> GraphDocument = StaticCastSharedRef<SGraphDocument>(InWidget))
+		{
+			GraphDocument->GraphEditor->GetViewLocation(ViewLocation, ZoomAmount);
+		}
+	
 		return TInstancedStruct<FGraphDocumentState>::Make(InContext.Object, ViewLocation, ZoomAmount);
 	});
 	Args.OnSetDocumentState = FOnSetDocumentState::CreateLambda([](const FWorkspaceEditorContext& InContext, TSharedRef<SWidget> InWidget, const TInstancedStruct<FWorkspaceDocumentState>& InDocumentState)
 	{
 		if(const FGraphDocumentState* GraphDocumentState = InDocumentState.GetPtr<FGraphDocumentState>())
 		{
-			TSharedRef<SGraphDocument> GraphDocument = StaticCastSharedRef<SGraphDocument>(InWidget);
-			GraphDocument->GraphEditor->SetViewLocation(GraphDocumentState->ViewLocation, GraphDocumentState->ZoomAmount);
+			if (const TSharedPtr<SGraphDocument> GraphDocument = StaticCastSharedRef<SGraphDocument>(InWidget))
+			{
+				GraphDocument->GraphEditor->SetViewLocation(GraphDocumentState->ViewLocation, GraphDocumentState->ZoomAmount);
+			}
 		}
 	});
 
@@ -149,7 +162,7 @@ void FWorkspaceEditorModule::OpenWorkspaceForObject(UObject* InObject, EOpenWork
 		}
 	}
 
-	UWorkspaceAssetEditor* WorkspaceEditor = nullptr;
+	IWorkspaceEditor* WorkspaceEditor = nullptr;
 
 	auto HandleNewWorkspace = [InObject, &WorkspaceEditor, WorkSpaceFactoryClass]()
 	{
@@ -166,8 +179,8 @@ void FWorkspaceEditorModule::OpenWorkspaceForObject(UObject* InObject, EOpenWork
 			UWorkspaceAssetEditor* AssetEditor = NewObject<UWorkspaceAssetEditor>(AssetEditorSubsystem, NAME_None, RF_Transient);
 			AssetEditor->SetObjectToEdit(NewWorkspace);
 			AssetEditor->Initialize();
-
-			WorkspaceEditor = AssetEditor;
+			
+			WorkspaceEditor = static_cast<IWorkspaceEditor*>(GEditor->GetEditorSubsystem<UAssetEditorSubsystem>()->FindEditorForAsset(NewWorkspace, true));
 		}
 	};
 
@@ -177,7 +190,7 @@ void FWorkspaceEditorModule::OpenWorkspaceForObject(UObject* InObject, EOpenWork
 		{
 			GEditor->GetEditorSubsystem<UAssetEditorSubsystem>()->OpenEditorForAsset(ExistingWorkspace);
 
-			WorkspaceEditor = static_cast<UWorkspaceAssetEditor*>(GEditor->GetEditorSubsystem<UAssetEditorSubsystem>()->FindEditorForAsset(ExistingWorkspace, true));
+			WorkspaceEditor = static_cast<IWorkspaceEditor*>(GEditor->GetEditorSubsystem<UAssetEditorSubsystem>()->FindEditorForAsset(ExistingWorkspace, true));
 		}
 	};
 
@@ -201,6 +214,11 @@ void FWorkspaceEditorModule::OpenWorkspaceForObject(UObject* InObject, EOpenWork
 
 		WorkspacePicker->ShowModal();
 	}
+
+	if(WorkspaceEditor)
+	{
+		WorkspaceEditor->OpenAssets({InObject});
+	}
 }
 
 const FObjectDocumentArgs* FWorkspaceEditorModule::FindObjectDocumentType(const FTopLevelAssetPath& InClassPath) const
@@ -226,6 +244,24 @@ void FWorkspaceEditorModule::ApplyWorkspaceDetailsCustomization(TSharedPtr<IDeta
 	}
 }
 
+void FWorkspaceEditorModule::RegisterWorkspaceItemDetails(const FOutlinerItemDetailsId& InItemDetailsId, TSharedPtr<IWorkspaceOutlinerItemDetails> InItemDetails)
+{
+	if (!OutlinerItemDetails.Contains(InItemDetailsId))
+	{
+		OutlinerItemDetails.Add(InItemDetailsId, InItemDetails);
+	}
+}
+
+void FWorkspaceEditorModule::UnregisterWorkspaceItemDetails(const FOutlinerItemDetailsId& InItemDetailsId)
+{
+	OutlinerItemDetails.Remove(InItemDetailsId);
+}
+
+TSharedPtr<IWorkspaceOutlinerItemDetails> FWorkspaceEditorModule::GetOutlinerItemDetails(const FOutlinerItemDetailsId& InItemDetailsId)
+{
+	TSharedPtr<IWorkspaceOutlinerItemDetails>* FoundDetails = OutlinerItemDetails.Find(InItemDetailsId);
+	return FoundDetails ? *FoundDetails : nullptr;
+}
 }
 
 IMPLEMENT_MODULE(UE::Workspace::FWorkspaceEditorModule, WorkspaceEditor);
diff --git a/Engine/Plugins/Experimental/Workspace/Source/WorkspaceEditor/Private/WorkspaceEditorModule.h b/Engine/Plugins/Experimental/Workspace/Source/WorkspaceEditor/Private/WorkspaceEditorModule.h
--- a/Engine/Plugins/Experimental/Workspace/Source/WorkspaceEditor/Private/WorkspaceEditorModule.h
+++ b/Engine/Plugins/Experimental/Workspace/Source/WorkspaceEditor/Private/WorkspaceEditorModule.h
@@ -1,16 +1,19 @@
-// Copyright Epic Games, Inc. All Rights Reserved.
+﻿// Copyright Epic Games, Inc. All Rights Reserved.
 
 #pragma once
 
 #include "IWorkspaceEditorModule.h"
 #include "UObject/TopLevelAssetPath.h"
+#include "IWorkspaceOutlinerItemDetails.h"
 
 struct FWorkspaceAssetRegistryExports;
+class SWorkspaceTabWrapper;
 
 namespace UE::Workspace
 {
 	struct FAssetDocumentSummoner;
 	class FWorkspaceEditor;
+	class FWorkspaceOutlinerMode;
 }
 
 namespace UE::Workspace
@@ -23,8 +26,10 @@ private:
 	virtual void RegisterObjectDocumentType(const FTopLevelAssetPath& InClassPath, const FObjectDocumentArgs& InParams) override;
 	virtual void UnregisterObjectDocumentType(const FTopLevelAssetPath& InClassPath) override;
 	virtual FObjectDocumentArgs CreateGraphDocumentArgs(const FGraphDocumentWidgetArgs& InArgs) override;
-	virtual void OpenWorkspaceForObject(UObject* InObject, EOpenWorkspaceMethod InOpenMethod, const TSubclassOf<UWorkspaceFactory> WorkSpaceFactoryClass) override;
+	virtual void OpenWorkspaceForObject(UObject* InObject, EOpenWorkspaceMethod InOpenMethod, const TSubclassOf<UWorkspaceFactory> WorkSpaceFactoryClass) override;	
 	virtual FOnRegisterDetailCustomizations& OnRegisterWorkspaceDetailsCustomization() override;
+	virtual void RegisterWorkspaceItemDetails(const FOutlinerItemDetailsId& InItemDetailsId, TSharedPtr<IWorkspaceOutlinerItemDetails> InItemDetails) override;	
+	virtual void UnregisterWorkspaceItemDetails(const FOutlinerItemDetailsId& InItemDetails) override;
 
 	// Find an existing registered object document type
 	const FObjectDocumentArgs* FindObjectDocumentType(const FTopLevelAssetPath& InClassPath) const;
@@ -41,12 +46,20 @@ private:
 	TMap<FTopLevelAssetPath, FObjectDocumentArgs> ObjectDocumentArgs;
 
 	TMap<FName, TSet<FTopLevelAssetPath>> DocumentAreaMap;
+	
+	static TMap<FOutlinerItemDetailsId, TSharedPtr<IWorkspaceOutlinerItemDetails>> OutlinerItemDetails;
+	static TSharedPtr<IWorkspaceOutlinerItemDetails> GetOutlinerItemDetails(const FOutlinerItemDetailsId& DetailId);
 
 	/** Event called to allow external clients to register details customizations */
 	FOnRegisterDetailCustomizations OnRegisterDetailCustomizations;
 
 	friend struct FAssetDocumentSummoner;
+	friend struct FWorkspaceOutlinerTreeItem;
+	friend class SWorkspaceOutlinerTreelabel;
 	friend class FWorkspaceEditor;
+	friend class FWorkspaceOutlinerMode;
+	friend class FWorkspaceOutlinerSourceControlColumn;
+	friend SWorkspaceTabWrapper;
 };
 
 }
diff --git a/Engine/Plugins/Experimental/Workspace/Source/WorkspaceEditor/Public/IWorkspaceEditorModule.h b/Engine/Plugins/Experimental/Workspace/Source/WorkspaceEditor/Public/IWorkspaceEditorModule.h
--- a/Engine/Plugins/Experimental/Workspace/Source/WorkspaceEditor/Public/IWorkspaceEditorModule.h
+++ b/Engine/Plugins/Experimental/Workspace/Source/WorkspaceEditor/Public/IWorkspaceEditorModule.h
@@ -8,6 +8,7 @@
 #include "WorkspaceFactory.h"
 #include "Misc/Attribute.h"
 #include "Modules/ModuleInterface.h"
+#include "IWorkspaceOutlinerItemDetails.h"
 
 class SWidget;
 class UEdGraph;
@@ -47,6 +48,15 @@ struct FWorkspaceEditorContext
 	UObject* Object;
 };
 
+struct FWorkspaceBreadcrumb : TSharedFromThis<FWorkspaceBreadcrumb>
+{
+	using FOnGetBreadcrumbLabel = TDelegate<TAttribute<FText>()>;
+	using FOnBreadcrumbClicked = TDelegate<void()>;
+	
+	FOnGetBreadcrumbLabel OnGetLabel;
+	FOnBreadcrumbClicked OnClicked;
+};
+
 using FOnMakeDocumentWidget = TDelegate<TSharedRef<SWidget>(const FWorkspaceEditorContext&)>;
 
 using FOnGetTabIcon = TDelegate<const FSlateBrush*(const FWorkspaceEditorContext&)>;
@@ -57,6 +67,8 @@ using FOnGetDocumentState = TDelegate<TInstancedStruct<FWorkspaceDocumentState>(
 
 using FOnSetDocumentState = TDelegate<void(const FWorkspaceEditorContext&, TSharedRef<SWidget>, const TInstancedStruct<FWorkspaceDocumentState>&)>;
 
+using FOnGetDocumentBreadcrumbTrail = TDelegate<void(const FWorkspaceEditorContext&, TArray<TSharedPtr<FWorkspaceBreadcrumb>>&)>;
+
 // Arguments used to make document widgets for objects
 struct FObjectDocumentArgs
 {
@@ -84,6 +96,9 @@ struct FObjectDocumentArgs
 
 	// Where to spawn the widget in the workspace layout - e.g. one of WorkspaceTabs
 	FName SpawnLocation = WorkspaceTabs::MiddleDocumentArea;
+
+	// Delegate called to get the bread crumb trail for this document tab
+	FOnGetDocumentBreadcrumbTrail OnGetDocumentBreadcrumbTrail;
 };
 
 
@@ -162,6 +177,9 @@ public:
 	// Event to allow registering details customizations
 	DECLARE_EVENT_OneParam(IWorkspaceEditorModule, FOnRegisterDetailCustomizations, TSharedPtr<IDetailsView>&);
 	virtual FOnRegisterDetailCustomizations& OnRegisterWorkspaceDetailsCustomization() = 0;
+
+	virtual void RegisterWorkspaceItemDetails(const FOutlinerItemDetailsId& InDetailsId, TSharedPtr<IWorkspaceOutlinerItemDetails> InDetails) = 0;
+	virtual void UnregisterWorkspaceItemDetails(const FOutlinerItemDetailsId& InDetailsId) = 0;
 };
 
 }
diff --git a/Engine/Plugins/Experimental/Workspace/Source/WorkspaceEditor/Public/IWorkspaceOutlinerItemDetails.h b/Engine/Plugins/Experimental/Workspace/Source/WorkspaceEditor/Public/IWorkspaceOutlinerItemDetails.h
new file mode 100644
--- /dev/null
+++ b/Engine/Plugins/Experimental/Workspace/Source/WorkspaceEditor/Public/IWorkspaceOutlinerItemDetails.h
@@ -0,0 +1,29 @@
+﻿// Copyright Epic Games, Inc. All Rights Reserved.
+
+#pragma once
+
+#include "WorkspaceAssetRegistryInfo.h"
+
+class UPackage;
+struct FToolMenuContext;
+struct FSlateBrush;
+
+namespace UE::Workspace
+{
+
+typedef FName FOutlinerItemDetailsId;
+static FOutlinerItemDetailsId MakeOutlinerDetailsId(const FWorkspaceOutlinerItemExport& InExport)
+{
+    return InExport.Data.IsValid() ? InExport.Data.GetScriptStruct()->GetFName() : NAME_None;
+}
+
+class IWorkspaceOutlinerItemDetails : public TSharedFromThis<IWorkspaceOutlinerItemDetails>
+{
+public:
+    virtual ~IWorkspaceOutlinerItemDetails() = default;
+    virtual const FSlateBrush* GetItemIcon() const { return nullptr; }
+    virtual void HandleDoubleClick(const FToolMenuContext& ToolMenuContext) const {}
+    virtual UPackage* GetPackage(const FWorkspaceOutlinerItemExport& Export) const { return nullptr; }
+};
+
+}
diff --git a/Engine/Plugins/Experimental/Workspace/Source/WorkspaceEditor/Public/WorkspaceAssetRegistryInfo.h b/Engine/Plugins/Experimental/Workspace/Source/WorkspaceEditor/Public/WorkspaceAssetRegistryInfo.h
new file mode 100644
--- /dev/null
+++ b/Engine/Plugins/Experimental/Workspace/Source/WorkspaceEditor/Public/WorkspaceAssetRegistryInfo.h
@@ -0,0 +1,57 @@
+﻿// Copyright Epic Games, Inc. All Rights Reserved.
+
+#pragma once
+
+#include "StructUtils/Public/InstancedStruct.h"
+#include "WorkspaceAssetRegistryInfo.generated.h"
+
+USTRUCT()
+struct FWorkspaceOutlinerItemData
+{
+	GENERATED_BODY()
+
+	FWorkspaceOutlinerItemData() = default;
+};
+
+USTRUCT()
+struct FWorkspaceOutlinerItemExport
+{
+	GENERATED_BODY()
+
+	FWorkspaceOutlinerItemExport() = default;
+
+	UPROPERTY()
+	FName Identifier;
+
+	UPROPERTY()
+	FName ParentIdentifier;
+
+	UPROPERTY()	
+	FSoftObjectPath AssetPath;
+
+	UPROPERTY()
+	TInstancedStruct<FWorkspaceOutlinerItemData> Data;
+
+	friend uint32 GetTypeHash(const FWorkspaceOutlinerItemExport& Export)
+	{
+		return HashCombine(GetTypeHash(Export.AssetPath), GetTypeHash(Export.Identifier));
+	}
+};
+
+namespace UE::Workspace
+{
+
+static const FLazyName ExportsWorkspaceItemsRegistryTag = TEXT("WorkspaceItemExports");
+
+}
+
+USTRUCT()
+struct FWorkspaceOutlinerItemExports
+{
+	GENERATED_BODY()
+
+	FWorkspaceOutlinerItemExports() = default;
+
+	UPROPERTY()
+	TArray<FWorkspaceOutlinerItemExport> Exports;
+};
diff --git a/Engine/Plugins/Experimental/Workspace/Source/WorkspaceEditor/Public/WorkspaceItemMenuContext.h b/Engine/Plugins/Experimental/Workspace/Source/WorkspaceEditor/Public/WorkspaceItemMenuContext.h
new file mode 100644
--- /dev/null
+++ b/Engine/Plugins/Experimental/Workspace/Source/WorkspaceEditor/Public/WorkspaceItemMenuContext.h
@@ -0,0 +1,14 @@
+// Copyright Epic Games, Inc. All Rights Reserved.
+
+#pragma once
+
+#include "WorkspaceAssetRegistryInfo.h"
+#include "WorkspaceItemMenuContext.generated.h"
+
+UCLASS()
+class WORKSPACEEDITOR_API UWorkspaceItemMenuContext : public UObject
+{
+	GENERATED_BODY()
+public:
+	TArray<FWorkspaceOutlinerItemExport> SelectedExports;
+};
diff --git a/Engine/Plugins/Experimental/Workspace/Source/WorkspaceEditor/WorkspaceEditor.Build.cs b/Engine/Plugins/Experimental/Workspace/Source/WorkspaceEditor/WorkspaceEditor.Build.cs
--- a/Engine/Plugins/Experimental/Workspace/Source/WorkspaceEditor/WorkspaceEditor.Build.cs
+++ b/Engine/Plugins/Experimental/Workspace/Source/WorkspaceEditor/WorkspaceEditor.Build.cs
@@ -36,6 +36,9 @@ namespace UnrealBuildTool.Rules
 					"EditorWidgets",
 					"Json",
 					"JsonUtilities",
+					"AssetRegistry",
+					"ContentBrowser",
+					"SceneOutliner"
 				}
 			);
 		}
