Change workspace navigation key mappings, pending other changes

diff --git a/Engine/Plugins/Experimental/Workspace/Source/WorkspaceEditor/Private/WorkspaceEditorCommands.cpp b/Engine/Plugins/Experimental/Workspace/Source/WorkspaceEditor/Private/WorkspaceEditorCommands.cpp
--- a/Engine/Plugins/Experimental/Workspace/Source/WorkspaceEditor/Private/WorkspaceEditorCommands.cpp
+++ b/Engine/Plugins/Experimental/Workspace/Source/WorkspaceEditor/Private/WorkspaceEditorCommands.cpp
@@ -24,9 +24,9 @@ namespace UE::Workspace
 	void FWorkspaceAssetEditorCommands::RegisterCommands()
 	{
 		UI_COMMAND(NavigateBackward, "Navigate Backward", "Moves backwards to previous location",
-			EUserInterfaceActionType::Button, FInputChord(EKeys::ThumbMouseButton), FInputChord(EKeys::Left, EModifierKey::Alt));
+			EUserInterfaceActionType::Button, FInputChord(EKeys::ThumbMouseButton), FInputChord(EKeys::B, EModifierKey::Alt));
 		UI_COMMAND(NavigateForward, "Navigate Forward", "Moves forwards to previous location",
-			EUserInterfaceActionType::Button, FInputChord(EKeys::ThumbMouseButton2), FInputChord(EKeys::Right, EModifierKey::Alt));
+			EUserInterfaceActionType::Button, FInputChord(EKeys::ThumbMouseButton2), FInputChord(EKeys::F, EModifierKey::Alt));
 	}
 
 }  // namespace UE::Workspace
