CIS Mac fix following 33296442

diff --git a/Engine/Plugins/Experimental/Workspace/Source/WorkspaceEditor/Private/Outliner/WorkspaceOutlinerHierarchy.cpp b/Engine/Plugins/Experimental/Workspace/Source/WorkspaceEditor/Private/Outliner/WorkspaceOutlinerHierarchy.cpp
--- a/Engine/Plugins/Experimental/Workspace/Source/WorkspaceEditor/Private/Outliner/WorkspaceOutlinerHierarchy.cpp
+++ b/Engine/Plugins/Experimental/Workspace/Source/WorkspaceEditor/Private/Outliner/WorkspaceOutlinerHierarchy.cpp
@@ -29,7 +29,7 @@ namespace UE::Workspace
 					FWorkspaceOutlinerItemExports::StaticStruct()->ImportText(*TagValue, &Exports, nullptr, 0, nullptr, FWorkspaceOutlinerItemExports::StaticStruct()->GetName());
 					for (const FWorkspaceOutlinerItemExport& Export : Exports.Exports)
 					{
-						if (FSceneOutlinerTreeItemPtr Item = Mode->CreateItemFor<FWorkspaceOutlinerTreeItem>(FWorkspaceOutlinerTreeItem::FItemData(Export)))
+						if (FSceneOutlinerTreeItemPtr Item = Mode->CreateItemFor<FWorkspaceOutlinerTreeItem>(FWorkspaceOutlinerTreeItem::FItemData{Export}))
 						{
 							OutItems.Add(Item);
 						}
@@ -70,7 +70,7 @@ namespace UE::Workspace
 									return ItemExport.Identifier == ParentIdentifier;
 								}))
 								{
-									Mode->CreateItemFor<FWorkspaceOutlinerTreeItem>(FWorkspaceOutlinerTreeItem::FItemData(*ExportPtr), true);
+									Mode->CreateItemFor<FWorkspaceOutlinerTreeItem>(FWorkspaceOutlinerTreeItem::FItemData{*ExportPtr}, true);
 								}
 							}
 						}
