Change SharedPtr capture to WeakPtr

diff --git a/Engine/Plugins/Experimental/Workspace/Source/WorkspaceEditor/Private/WorkspaceEditorModule.cpp b/Engine/Plugins/Experimental/Workspace/Source/WorkspaceEditor/Private/WorkspaceEditorModule.cpp
--- a/Engine/Plugins/Experimental/Workspace/Source/WorkspaceEditor/Private/WorkspaceEditorModule.cpp
+++ b/Engine/Plugins/Experimental/Workspace/Source/WorkspaceEditor/Private/WorkspaceEditorModule.cpp
@@ -47,6 +47,7 @@ FObjectDocumentArgs FWorkspaceEditorModule::CreateGraphDocumentArgs(const FGraph
 	FObjectDocumentArgs Args;
 	Args.OnMakeDocumentWidget = FOnMakeDocumentWidget::CreateLambda([InArgs](const FWorkspaceEditorContext& InContext)
 	{
+		TWeakPtr<FWorkspaceEditor> WeakWorkspaceEditor = StaticCastSharedRef<FWorkspaceEditor>(InContext.WorkspaceEditor);
 		return SNew(SGraphDocument, StaticCastSharedRef<FWorkspaceEditor>(InContext.WorkspaceEditor), CastChecked<UEdGraph>(InContext.Object))
 			.OnCreateActionMenu_Lambda([OnCreateActionMenu = InArgs.OnCreateActionMenu](const FWorkspaceEditorContext& InContext, UEdGraph* InGraph, const FVector2D& InNodePosition, const TArray<UEdGraphPin*>& InDraggedPins, bool bInAutoExpand, SGraphEditor::FActionMenuClosed InOnMenuClosed)
 			{
@@ -68,8 +69,8 @@ FObjectDocumentArgs FWorkspaceEditorModule::CreateGraphDocumentArgs(const FGraph
 			.OnPasteNodes(InArgs.OnPasteNodes)
 			.OnCanDuplicateSelectedNodes(InArgs.OnCanDuplicateSelectedNodes)
 			.OnDuplicateSelectedNodes(InArgs.OnDuplicateSelectedNodes)
-			.OnNavigateHistoryForward_Lambda([WorkspaceEditor=StaticCastSharedRef<FWorkspaceEditor>(InContext.WorkspaceEditor)](){ WorkspaceEditor->NavigateForward(); })
-			.OnNavigateHistoryBack_Lambda([WorkspaceEditor=StaticCastSharedRef<FWorkspaceEditor>(InContext.WorkspaceEditor)](){ WorkspaceEditor->NavigateBack(); })
+			.OnNavigateHistoryForward_Lambda([WeakWorkspaceEditor](){ if(const TSharedPtr<FWorkspaceEditor> SharedWorkspaceEditor = WeakWorkspaceEditor.Pin()) SharedWorkspaceEditor->NavigateForward(); })
+			.OnNavigateHistoryBack_Lambda([WeakWorkspaceEditor](){ if(const TSharedPtr<FWorkspaceEditor> SharedWorkspaceEditor = WeakWorkspaceEditor.Pin()) SharedWorkspaceEditor->NavigateBack(); })
 		
 		;
 	});
