Add FileState column to Workspace outliner

diff --git a/Engine/Plugins/Experimental/Workspace/Source/WorkspaceEditor/Private/Outliner/WorkspaceOutlinerColumns.cpp b/Engine/Plugins/Experimental/Workspace/Source/WorkspaceEditor/Private/Outliner/WorkspaceOutlinerColumns.cpp
--- a/Engine/Plugins/Experimental/Workspace/Source/WorkspaceEditor/Private/Outliner/WorkspaceOutlinerColumns.cpp
+++ b/Engine/Plugins/Experimental/Workspace/Source/WorkspaceEditor/Private/Outliner/WorkspaceOutlinerColumns.cpp
@@ -10,6 +10,7 @@
 #include "ISourceControlProvider.h"
 #include "SourceControlHelpers.h"
 #include "Widgets/Images/SImage.h"
+#include "UObject/Package.h"
 
 #define LOCTEXT_NAMESPACE "WorkspaceOutlinerColumns"
 
@@ -74,7 +75,7 @@ namespace UE::Workspace
 							return FText::Format(LOCTEXT("RevisionControlStatusFormat", "File: {0}\nStatus: {1}"), State->GetDisplayName(), State->GetDisplayTooltip());
 						}
 
-						return LOCTEXT("RevisionControlStatus", "Revision control status of this parameter");
+						return LOCTEXT("RevisionControlStatus", "Revision control status of this entry");
 					})
 					.Image_Lambda([this, PackagePath]() -> const FSlateBrush*
 					{
@@ -92,6 +93,93 @@ namespace UE::Workspace
 		
 		return SNullWidget::NullWidget;
 	}
+
+	FName WorkspaceOutlinerFileState("File State");
+	FName FWorkspaceOutlinerFileStateColumn::GetID()
+	{
+		return WorkspaceOutlinerFileState;
+	}
+
+	SHeaderRow::FColumn::FArguments FWorkspaceOutlinerFileStateColumn::ConstructHeaderRowColumn()
+	{
+		return SHeaderRow::Column(GetColumnID())
+		.FixedWidth(24.f)
+		.HAlignHeader(HAlign_Center)
+		.VAlignHeader(VAlign_Center)
+		.HAlignCell(HAlign_Center)
+		.VAlignCell(VAlign_Center)
+		[
+			SNew(SBox)
+			.WidthOverride(16.0f)
+			.HeightOverride(16.0f)
+			.VAlign(VAlign_Center)
+			.HAlign(HAlign_Center)
+			[
+				SNew(SImage)
+				.ColorAndOpacity(FSlateColor::UseForeground())
+				.Image(FAppStyle::GetBrush("ContentBrowser.ContentDirty"))
+				.ToolTipText(LOCTEXT("RevisionControlStatusTooltip", "File status of this entry"))
+			]
+		];	
+	}
+
+	const TSharedRef<SWidget> FWorkspaceOutlinerFileStateColumn::ConstructRowWidget(FSceneOutlinerTreeItemRef Item, const STableRow<FSceneOutlinerTreeItemPtr>& Row)
+	{
+		if (const FWorkspaceOutlinerTreeItem* TreeItem = Item->CastTo<FWorkspaceOutlinerTreeItem>())
+		{
+			auto GetWeakPackage = [TreeItem]()
+			{
+				TWeakObjectPtr<const UPackage> WeakPackage = nullptr;
+				if (const TSharedPtr<IWorkspaceOutlinerItemDetails> SharedFactory = FWorkspaceEditorModule::GetOutlinerItemDetails(MakeOutlinerDetailsId(TreeItem->Export)))
+				{
+					WeakPackage = SharedFactory->GetPackage(TreeItem->Export);
+				}
+				
+				if (!WeakPackage.IsValid())
+				{
+					if(const UObject* Object = TreeItem->Export.AssetPath.ResolveObject())
+					{
+						WeakPackage = Object->GetPackage();
+					}
+				}		
+
+				return WeakPackage;
+			};		
+			
+			return SNew(SBox)
+			.HAlign(HAlign_Center)
+			.VAlign(VAlign_Center)
+			.HeightOverride(20.0f)		
+			[
+				SNew(SImage)
+				.ToolTipText_Lambda([this, GetWeakPackage]() -> FText
+				{
+					TWeakObjectPtr<const UPackage> WeakPackage = GetWeakPackage();					
+					if (const UPackage* Package = WeakPackage.Get())
+					{
+						return FText::Format(LOCTEXT("FileStatusTooltipFormat", "File: {0}"), FText::FromString(SourceControlHelpers::PackageFilename(Package)));
+					}
+				
+					return LOCTEXT("FileStatusTooltip", "File status of this entry");
+				})
+				.Image_Lambda([this, GetWeakPackage]() -> const FSlateBrush*
+				{
+					TWeakObjectPtr<const UPackage> WeakPackage = GetWeakPackage();
+					if (const UPackage* Package = WeakPackage.Get())
+					{
+						if (Package->IsDirty())
+						{
+							return FAppStyle::GetBrush("ContentBrowser.ContentDirty");
+						}
+					}
+					
+					return nullptr;
+				})
+			];
+		}
+		
+		return SNullWidget::NullWidget;
+	}
 }
 
 #undef LOCTEXT_NAMESPACE // "WorkspaceOutlinerColumns"
\ No newline at end of file
diff --git a/Engine/Plugins/Experimental/Workspace/Source/WorkspaceEditor/Private/Outliner/WorkspaceOutlinerColumns.h b/Engine/Plugins/Experimental/Workspace/Source/WorkspaceEditor/Private/Outliner/WorkspaceOutlinerColumns.h
--- a/Engine/Plugins/Experimental/Workspace/Source/WorkspaceEditor/Private/Outliner/WorkspaceOutlinerColumns.h
+++ b/Engine/Plugins/Experimental/Workspace/Source/WorkspaceEditor/Private/Outliner/WorkspaceOutlinerColumns.h
@@ -28,6 +28,25 @@ public:
 	virtual bool SupportsSorting() const override { return false; }
 	// End ISceneOutlinerColumn overrides
 
+private:
+	TWeakPtr<ISceneOutliner> WeakSceneOutliner;
+};
+
+class FWorkspaceOutlinerFileStateColumn : public ISceneOutlinerColumn
+{
+public:
+	static FName GetID();
+	
+	FWorkspaceOutlinerFileStateColumn(ISceneOutliner& SceneOutliner) : WeakSceneOutliner(StaticCastSharedRef<ISceneOutliner>(SceneOutliner.AsShared())) {}
+	virtual ~FWorkspaceOutlinerFileStateColumn() override = default;
+	
+	// Begin ISceneOutlinerColumn overrides
+	virtual FName GetColumnID() override { return GetID(); }
+	virtual SHeaderRow::FColumn::FArguments ConstructHeaderRowColumn() override;
+	virtual const TSharedRef<SWidget> ConstructRowWidget(FSceneOutlinerTreeItemRef TreeItem, const STableRow<FSceneOutlinerTreeItemPtr>& Row) override;
+	virtual bool SupportsSorting() const override { return false; }
+	// End ISceneOutlinerColumn overrides
+
 private:
 	TWeakPtr<ISceneOutliner> WeakSceneOutliner;
 };
diff --git a/Engine/Plugins/Experimental/Workspace/Source/WorkspaceEditor/Private/SWorkspaceView.cpp b/Engine/Plugins/Experimental/Workspace/Source/WorkspaceEditor/Private/SWorkspaceView.cpp
--- a/Engine/Plugins/Experimental/Workspace/Source/WorkspaceEditor/Private/SWorkspaceView.cpp
+++ b/Engine/Plugins/Experimental/Workspace/Source/WorkspaceEditor/Private/SWorkspaceView.cpp
@@ -26,7 +26,8 @@ void SWorkspaceView::Construct(const FArguments& InArgs, UWorkspace* InWorkspace
 	{
 		InitOptions.OutlinerIdentifier = TEXT("WorkspaceEditorOutliner");
 		InitOptions.bShowHeaderRow = true;
-		InitOptions.ColumnMap.Add(FSceneOutlinerBuiltInColumnTypes::Label(), FSceneOutlinerColumnInfo(ESceneOutlinerColumnVisibility::Visible, 0));
+		InitOptions.ColumnMap.Add(FSceneOutlinerBuiltInColumnTypes::Label(), FSceneOutlinerColumnInfo(ESceneOutlinerColumnVisibility::Visible, 10));
+		InitOptions.ColumnMap.Add(FWorkspaceOutlinerFileStateColumn::GetID(), FSceneOutlinerColumnInfo(ESceneOutlinerColumnVisibility::Visible, 0, FCreateSceneOutlinerColumn::CreateLambda([](ISceneOutliner& InSceneOutliner) { return MakeShareable(new FWorkspaceOutlinerFileStateColumn(InSceneOutliner)); }), false));
 		InitOptions.ColumnMap.Add(FWorkspaceOutlinerSourceControlColumn::GetID(), FSceneOutlinerColumnInfo(ESceneOutlinerColumnVisibility::Visible, 100, FCreateSceneOutlinerColumn::CreateLambda([](ISceneOutliner& InSceneOutliner) { return MakeShareable(new FWorkspaceOutlinerSourceControlColumn(InSceneOutliner)); }), false));	
 		InitOptions.ModeFactory = FCreateSceneOutlinerMode::CreateLambda([this, WeakWorkspaceEditor=InWorkspaceEditor.ToWeakPtr()](SSceneOutliner* InOutliner) { return new UE::Workspace::FWorkspaceOutlinerMode(UE::Workspace::FWorkspaceOutlinerMode(InOutliner, Workspace, WeakWorkspaceEditor)); });
 	}
diff --git a/Engine/Plugins/Experimental/Workspace/Source/WorkspaceEditor/Private/WorkspaceEditorModule.h b/Engine/Plugins/Experimental/Workspace/Source/WorkspaceEditor/Private/WorkspaceEditorModule.h
--- a/Engine/Plugins/Experimental/Workspace/Source/WorkspaceEditor/Private/WorkspaceEditorModule.h
+++ b/Engine/Plugins/Experimental/Workspace/Source/WorkspaceEditor/Private/WorkspaceEditorModule.h
@@ -59,6 +59,7 @@ private:
 	friend class FWorkspaceEditor;
 	friend class FWorkspaceOutlinerMode;
 	friend class FWorkspaceOutlinerSourceControlColumn;
+	friend class FWorkspaceOutlinerFileStateColumn;
 	friend SWorkspaceTabWrapper;
 };
 
