Move asset GUIDs asset-registry data from Workspace to WorkspaceAssetEntry

diff --git a/Engine/Plugins/Experimental/Workspace/Source/WorkspaceEditor/Private/Workspace.cpp b/Engine/Plugins/Experimental/Workspace/Source/WorkspaceEditor/Private/Workspace.cpp
--- a/Engine/Plugins/Experimental/Workspace/Source/WorkspaceEditor/Private/Workspace.cpp
+++ b/Engine/Plugins/Experimental/Workspace/Source/WorkspaceEditor/Private/Workspace.cpp
@@ -13,8 +13,6 @@
 
 REGISTER_ASSETDEPENDENCY_GATHERER(FExternalAssetDependencyGatherer, UWorkspace);
 
-const FName UWorkspace::ExportsAssetRegistryTag = TEXT("Exports");
-
 bool UWorkspace::AddAsset(const FAssetData& InAsset, bool bSetupUndoRedo, bool bPrintPythonCommand)
 {
 	if(!IsAssetSupported(InAsset))
@@ -357,25 +355,3 @@ void UWorkspace::PostDuplicate(bool bDuplicateForPIE)
 
 	Guid = FGuid::NewGuid();
 }
-
-void UWorkspace::GetAssetRegistryTags(FAssetRegistryTagsContext Context) const
-{
-	Super::GetAssetRegistryTags(Context);
-
-	FWorkspaceAssetRegistryExports Exports;
-	Exports.Assets.Reserve(AssetEntries.Num());
-
-	for(const UWorkspaceAssetEntry* AssetEntry : AssetEntries)
-	{
-		if (AssetEntry)
-		{
-			Exports.Assets.Emplace(AssetEntry->Asset.GetUniqueID());
-		}
-	
-	}
-
-	FString TagValue;
-	FWorkspaceAssetRegistryExports::StaticStruct()->ExportText(TagValue, &Exports, nullptr, nullptr, PPF_None, nullptr);
-
-	Context.AddTag(FAssetRegistryTag(ExportsAssetRegistryTag, TagValue, FAssetRegistryTag::TT_Hidden));
-}
diff --git a/Engine/Plugins/Experimental/Workspace/Source/WorkspaceEditor/Private/Workspace.h b/Engine/Plugins/Experimental/Workspace/Source/WorkspaceEditor/Private/Workspace.h
--- a/Engine/Plugins/Experimental/Workspace/Source/WorkspaceEditor/Private/Workspace.h
+++ b/Engine/Plugins/Experimental/Workspace/Source/WorkspaceEditor/Private/Workspace.h
@@ -74,8 +74,6 @@ class UWorkspace : public UObject
 	friend class UAssetDefinition_Workspace;
 	friend class UWorkspaceState;
 
-	WORKSPACEEDITOR_API static const FName ExportsAssetRegistryTag;
-
 	// Adds an asset to the workspace
 	// @return true if the asset was added
 	UFUNCTION(BlueprintCallable, Category = "Workspace")
@@ -120,7 +118,6 @@ class UWorkspace : public UObject
 	// UObject interface
 	virtual void PostLoad() override;
 	virtual void PostDuplicate(bool bDuplicateForPIE) override;
-	virtual void GetAssetRegistryTags(FAssetRegistryTagsContext Context) const override;
 	virtual void PostTransacted(const FTransactionObjectEvent& TransactionEvent) override;
 	virtual bool IsEditorOnly() const override { return true; }
 	virtual void Serialize(FArchive& Ar) override;
diff --git a/Engine/Plugins/Experimental/Workspace/Source/WorkspaceEditor/Private/WorkspaceAssetEntry.cpp b/Engine/Plugins/Experimental/Workspace/Source/WorkspaceEditor/Private/WorkspaceAssetEntry.cpp
--- a/Engine/Plugins/Experimental/Workspace/Source/WorkspaceEditor/Private/WorkspaceAssetEntry.cpp
+++ b/Engine/Plugins/Experimental/Workspace/Source/WorkspaceEditor/Private/WorkspaceAssetEntry.cpp
@@ -1,11 +1,21 @@
 ﻿// Copyright Epic Games, Inc. All Rights Reserved.
 
 #include "WorkspaceAssetEntry.h"
+#include "Workspace.h"
+#include "UObject/AssetRegistryTagsContext.h"
 #include "UObject/Package.h"
 
+const FName UWorkspaceAssetEntry::ExportsAssetRegistryTag = TEXT("Exports");
+
 bool UWorkspaceAssetEntry::IsAsset() const
 {
 	// Entries are considered assets to allow using the asset logic for save dialogs, etc.
 	// Also, they return true even if pending kill, in order to show up as deleted in these dialogs.
 	return IsPackageExternal() && !GetPackage()->HasAnyFlags(RF_Transient) && !HasAnyFlags(RF_Transient | RF_ClassDefaultObject);
 }
+
+void UWorkspaceAssetEntry::GetAssetRegistryTags(FAssetRegistryTagsContext Context) const
+{
+	UObject::GetAssetRegistryTags(Context);
+	Context.AddTag(FAssetRegistryTag(UWorkspaceAssetEntry::ExportsAssetRegistryTag, Asset.GetUniqueID().ToString(), FAssetRegistryTag::TT_Hidden));
+}
diff --git a/Engine/Plugins/Experimental/Workspace/Source/WorkspaceEditor/Private/WorkspaceAssetEntry.h b/Engine/Plugins/Experimental/Workspace/Source/WorkspaceEditor/Private/WorkspaceAssetEntry.h
--- a/Engine/Plugins/Experimental/Workspace/Source/WorkspaceEditor/Private/WorkspaceAssetEntry.h
+++ b/Engine/Plugins/Experimental/Workspace/Source/WorkspaceEditor/Private/WorkspaceAssetEntry.h
@@ -6,12 +6,17 @@
 #include "UObject/Object.h"
 #include "WorkspaceAssetEntry.generated.h"
 
+class FAssetRegistryTagsContext;
+
 UCLASS()
 class UWorkspaceAssetEntry : public UObject
 {
 	GENERATED_BODY()
 public:	
+	WORKSPACEEDITOR_API static const FName ExportsAssetRegistryTag;
+
 	virtual bool IsAsset() const override;
+	virtual void GetAssetRegistryTags(FAssetRegistryTagsContext Context) const override;
 
 	UPROPERTY()
 	TSoftObjectPtr<UObject> Asset;	
diff --git a/Engine/Plugins/Experimental/Workspace/Source/WorkspaceEditor/Private/WorkspaceEditorModule.cpp b/Engine/Plugins/Experimental/Workspace/Source/WorkspaceEditor/Private/WorkspaceEditorModule.cpp
--- a/Engine/Plugins/Experimental/Workspace/Source/WorkspaceEditor/Private/WorkspaceEditorModule.cpp
+++ b/Engine/Plugins/Experimental/Workspace/Source/WorkspaceEditor/Private/WorkspaceEditorModule.cpp
@@ -3,6 +3,7 @@
 #include "WorkspaceEditorModule.h"
 
 #include "EdGraphUtilities.h"
+#include "ExternalPackageHelper.h"
 #include "WorkspaceDocumentState.h"
 #include "GraphDocumentState.h"
 #include "SGraphDocument.h"
@@ -126,8 +127,26 @@ FObjectDocumentArgs FWorkspaceEditorModule::CreateGraphDocumentArgs(const FGraph
 
 bool FWorkspaceEditorModule::GetExportedAssetsForWorkspace(const FAssetData& InWorkspaceAsset, FWorkspaceAssetRegistryExports& OutExports)
 {
-	const FString TagValue = InWorkspaceAsset.GetTagValueRef<FString>(UWorkspace::ExportsAssetRegistryTag);
-	return FWorkspaceAssetRegistryExports::StaticStruct()->ImportText(*TagValue, &OutExports, nullptr, PPF_None, nullptr, FWorkspaceAssetRegistryExports::StaticStruct()->GetName()) != nullptr;
+	const FString ExternalObjectsPath = FExternalPackageHelper::GetExternalObjectsPath(InWorkspaceAsset.PackageName.ToString());
+	FARFilter Filter;
+	Filter.bRecursivePaths = true;
+	Filter.bIncludeOnlyOnDiskAssets = true;
+	Filter.PackagePaths.Add(*ExternalObjectsPath);
+
+	TArray<FAssetData> AssetDataEntries;
+	const IAssetRegistry& AssetRegistry = FAssetRegistryModule::GetRegistry();
+	AssetRegistry.GetAssets(Filter, AssetDataEntries);
+
+	for (const FAssetData& AssetDataEntry : AssetDataEntries)
+	{
+		FSoftObjectPath SoftObjectPath(AssetDataEntry.GetTagValueRef<FString>(UWorkspaceAssetEntry::ExportsAssetRegistryTag));
+		if(SoftObjectPath.IsValid())
+		{
+			OutExports.Assets.Add({SoftObjectPath});
+		}
+	}
+
+	return OutExports.Assets.Num() > 0;
 }
 
 void FWorkspaceEditorModule::OpenWorkspaceForObject(UObject* InObject, EOpenWorkspaceMethod InOpenMethod, const TSubclassOf<UWorkspaceFactory> WorkSpaceFactoryClass/*= UWorkspaceFactory::StaticClass()*/)
