Add AssetRegistry update to AnimGraph editor data on-modified

diff --git a/Engine/Plugins/Experimental/AnimNext/Source/AnimNextUncookedOnly/Private/Graph/AnimNextGraph_EditorData.cpp b/Engine/Plugins/Experimental/AnimNext/Source/AnimNextUncookedOnly/Private/Graph/AnimNextGraph_EditorData.cpp
--- a/Engine/Plugins/Experimental/AnimNext/Source/AnimNextUncookedOnly/Private/Graph/AnimNextGraph_EditorData.cpp
+++ b/Engine/Plugins/Experimental/AnimNext/Source/AnimNextUncookedOnly/Private/Graph/AnimNextGraph_EditorData.cpp
@@ -3,6 +3,7 @@
 #include "Graph/AnimNextGraph_EditorData.h"
 
 #include "UncookedOnlyUtils.h"
+#include "AssetRegistry/IAssetRegistry.h"
 #include "Curves/CurveFloat.h"
 #include "Graph/AnimNextGraph.h"
 #include "Graph/AnimNextGraph_AnimationGraph.h"
@@ -14,6 +15,7 @@
 #include "RigVMModel/RigVMFunctionLibrary.h"
 #include "RigVMModel/RigVMNotifications.h"
 #include "RigVMModel/Nodes/RigVMCollapseNode.h"
+#include "UObject/AssetRegistryTagsContext.h"
 #include "UObject/FortniteMainBranchObjectVersion.h"
 #include "UObject/LinkerLoad.h"
 
@@ -132,6 +134,11 @@ void UAnimNextGraph_EditorData::RecompileVM()
 	CachedExports = FAnimNextParameterProviderAssetRegistryExports();
 	UE::AnimNext::UncookedOnly::FUtils::GetAssetParameters(this, CachedExports.GetValue());
 	UE::AnimNext::UncookedOnly::FUtils::Compile(GetTypedOuter<UAnimNextGraph>());
+	
+	if (IAssetRegistry* AssetRegistry = IAssetRegistry::Get())
+	{
+		AssetRegistry->AssetUpdateTags(GetTypedOuter<UAnimNextGraph>(), EAssetRegistryTagsCaller::Fast);
+	}
 }
 
 void UAnimNextGraph_EditorData::HandleModifiedEvent(ERigVMGraphNotifType InNotifType, URigVMGraph* InGraph, UObject* InSubject)
diff --git a/Engine/Plugins/Experimental/Workspace/Source/WorkspaceEditor/Private/Outliner/WorkspaceOutlinerMode.cpp b/Engine/Plugins/Experimental/Workspace/Source/WorkspaceEditor/Private/Outliner/WorkspaceOutlinerMode.cpp
--- a/Engine/Plugins/Experimental/Workspace/Source/WorkspaceEditor/Private/Outliner/WorkspaceOutlinerMode.cpp
+++ b/Engine/Plugins/Experimental/Workspace/Source/WorkspaceEditor/Private/Outliner/WorkspaceOutlinerMode.cpp
@@ -22,6 +22,11 @@ FWorkspaceOutlinerMode::FWorkspaceOutlinerMode(SSceneOutliner* InSceneOutliner,
 	{
 		Workspace->ModifiedDelegate.AddRaw(this, &FWorkspaceOutlinerMode::OnWorkspaceModified);
 	}
+
+	if (FAssetRegistryModule* AssetRegistryModule = FModuleManager::LoadModulePtr<FAssetRegistryModule>(AssetRegistryConstants::ModuleName))
+	{
+		AssetRegistryModule->Get().OnAssetUpdated().AddRaw(this, &FWorkspaceOutlinerMode::OnAssetRegistryAssetUpdate);
+	}
 }
 
 FWorkspaceOutlinerMode::~FWorkspaceOutlinerMode()
@@ -30,6 +35,11 @@ FWorkspaceOutlinerMode::~FWorkspaceOutlinerMode()
 	{
 		Workspace->ModifiedDelegate.RemoveAll(this);
 	}
+
+	if (FAssetRegistryModule* AssetRegistryModule = FModuleManager::LoadModulePtr<FAssetRegistryModule>(AssetRegistryConstants::ModuleName))
+	{
+		AssetRegistryModule->Get().OnAssetUpdated().RemoveAll(this);
+	}
 }
 
 void FWorkspaceOutlinerMode::Rebuild()
@@ -268,6 +278,10 @@ void FWorkspaceOutlinerMode::DeleteItems(TArrayView<const FSceneOutlinerTreeItem
 	}
 }
 
+void FWorkspaceOutlinerMode::OnAssetRegistryAssetUpdate(const FAssetData& AssetData)
+{
+	SceneOutliner->FullRefresh();
+}
 }	
 
 #undef LOCTEXT_NAMESPACE // "FWorkspaceOutlinerMode"
\ No newline at end of file
diff --git a/Engine/Plugins/Experimental/Workspace/Source/WorkspaceEditor/Private/Outliner/WorkspaceOutlinerMode.h b/Engine/Plugins/Experimental/Workspace/Source/WorkspaceEditor/Private/Outliner/WorkspaceOutlinerMode.h
--- a/Engine/Plugins/Experimental/Workspace/Source/WorkspaceEditor/Private/Outliner/WorkspaceOutlinerMode.h
+++ b/Engine/Plugins/Experimental/Workspace/Source/WorkspaceEditor/Private/Outliner/WorkspaceOutlinerMode.h
@@ -33,6 +33,8 @@ namespace UE::Workspace
 
 		void OpenItems(TArrayView<const FSceneOutlinerTreeItemPtr> Items) const;
 		void DeleteItems(TArrayView<const FSceneOutlinerTreeItemPtr> Items) const;
+
+		void OnAssetRegistryAssetUpdate(const FAssetData& AssetData);
 	private:
 		TWeakObjectPtr<UWorkspace> WeakWorkspace;
 		TWeakPtr<IWorkspaceEditor> WeakWorkspaceEditor;
