Fix duplicate LOCTEXT entries in WorkspaceOutlinerColumns

diff --git a/Engine/Plugins/Experimental/Workspace/Source/WorkspaceEditor/Private/Outliner/WorkspaceOutlinerColumns.cpp b/Engine/Plugins/Experimental/Workspace/Source/WorkspaceEditor/Private/Outliner/WorkspaceOutlinerColumns.cpp
--- a/Engine/Plugins/Experimental/Workspace/Source/WorkspaceEditor/Private/Outliner/WorkspaceOutlinerColumns.cpp
+++ b/Engine/Plugins/Experimental/Workspace/Source/WorkspaceEditor/Private/Outliner/WorkspaceOutlinerColumns.cpp
@@ -118,7 +118,7 @@ namespace UE::Workspace
 				SNew(SImage)
 				.ColorAndOpacity(FSlateColor::UseForeground())
 				.Image(FAppStyle::GetBrush("ContentBrowser.ContentDirty"))
-				.ToolTipText(LOCTEXT("RevisionControlStatusTooltip", "File status of this entry"))
+				.ToolTipText(LOCTEXT("FileStatusTooltip", "File status of this entry"))
 			]
 		];	
 	}
