AnimNext: Allow external tabs to be persisted to the workspace layout

Adds a layout extension delegate so external users of the workspace systems can persist tabs in the workspace layout
Adds extra tab areas (we have 6 in total in a grid now)
Persists the layout of the AnimNext compiler results and trait stack editor

diff --git a/Engine/Plugins/Experimental/AnimNext/Source/AnimNextEditor/Private/AnimNextEditorModule.cpp b/Engine/Plugins/Experimental/AnimNext/Source/AnimNextEditor/Private/AnimNextEditorModule.cpp
--- a/Engine/Plugins/Experimental/AnimNext/Source/AnimNextEditor/Private/AnimNextEditorModule.cpp
+++ b/Engine/Plugins/Experimental/AnimNext/Source/AnimNextEditor/Private/AnimNextEditorModule.cpp
@@ -46,6 +46,7 @@
 #include "Graph/AnimNextGraph_OutlinerItemDetails.h"
 #include "Param/AnimNextActorLocatorEditor.h"
 #include "IWorkspaceEditor.h"
+#include "Framework/Docking/LayoutExtender.h"
 
 #define LOCTEXT_NAMESPACE "AnimNextEditorModule"
 
@@ -93,18 +94,26 @@ void FModule::StartupModule()
 
 	Workspace::IWorkspaceEditorModule& WorkspaceEditorModule = FModuleManager::Get().LoadModuleChecked<Workspace::IWorkspaceEditorModule>("WorkspaceEditor");
 
-	WorkspaceEditorModule.OnRegisterTabsForEditor().AddLambda([](FWorkflowAllowedTabSet& TabFactories, const TSharedRef<FTabManager>& InTabManager, TSharedPtr<UE::Workspace::IWorkspaceEditor> EditorPtr)
+	WorkspaceEditorModule.OnRegisterTabsForEditor().AddLambda([](FWorkflowAllowedTabSet& TabFactories, const TSharedRef<FTabManager>& InTabManager, TSharedPtr<UE::Workspace::IWorkspaceEditor> InEditorPtr)
 		{
-			TabFactories.RegisterFactory(MakeShared<UE::AnimNext::Editor::FTraitEditorTabSummoner>(EditorPtr));
+			TSharedRef<FTraitEditorTabSummoner> TraitEditorTabSummoner = MakeShared<FTraitEditorTabSummoner>(InEditorPtr);
+			TabFactories.RegisterFactory(TraitEditorTabSummoner);
+			TraitEditorTabSummoner->RegisterTabSpawner(InTabManager, nullptr);
 
-			TabFactories.RegisterFactory(MakeShared<UE::AnimNext::Editor::FAnimNextCompilerResultsTabSummoner>(EditorPtr));
-
-			for (auto FactoryIt = TabFactories.CreateIterator(); FactoryIt; ++FactoryIt)
-			{
-				FactoryIt.Value()->RegisterTabSpawner(InTabManager, nullptr);
-			}
+			TSharedRef<FAnimNextCompilerResultsTabSummoner> CompilerResultsTabSummoner = MakeShared<FAnimNextCompilerResultsTabSummoner>(InEditorPtr);
+			TabFactories.RegisterFactory(CompilerResultsTabSummoner);
+			CompilerResultsTabSummoner->RegisterTabSpawner(InTabManager, nullptr);
 		});
 
+	WorkspaceEditorModule.OnExtendTabs().AddLambda([](FLayoutExtender& InLayoutExtender, TSharedPtr<UE::Workspace::IWorkspaceEditor> InEditorPtr)
+	{
+		FTabManager::FTab TraitEditorTab(FTabId(TraitEditorTabName), ETabState::ClosedTab);
+		InLayoutExtender.ExtendLayout(FTabId(Workspace::WorkspaceTabs::TopRightDocumentArea), ELayoutExtensionPosition::After, TraitEditorTab);
+
+		FTabManager::FTab CompilerResultsTab(FTabId(CompilerResultsTabName), ETabState::ClosedTab);
+		InLayoutExtender.ExtendLayout(FTabId(Workspace::WorkspaceTabs::BottomMiddleDocumentArea), ELayoutExtensionPosition::After, CompilerResultsTab);
+	});
+
 	RegisterWorkspaceDocumentTypes(WorkspaceEditorModule);
 
 	WorkspaceEditorModule.OnRegisterWorkspaceDetailsCustomization().AddLambda([](const TWeakPtr<Workspace::IWorkspaceEditor>& InWorkspaceEditor, TSharedPtr<IDetailsView>& InDetailsView)
@@ -303,7 +312,7 @@ void FModule::RegisterWorkspaceDocumentTypes(Workspace::IWorkspaceEditorModule&
 			DetailsView->SetObject(Schedule);
 			return DetailsView;
 		}
-	), Workspace::WorkspaceTabs::MiddleDocumentArea);
+	), Workspace::WorkspaceTabs::TopMiddleDocumentArea);
 	ScheduleDocumentArgs.OnGetTabName = Workspace::FOnGetTabName::CreateLambda([](const Workspace::FWorkspaceEditorContext& InContext)
 	{
 		const UAnimNextSchedule* Schedule = CastChecked<UAnimNextSchedule>(InContext.Object);
@@ -380,7 +389,7 @@ void FModule::RegisterWorkspaceDocumentTypes(Workspace::IWorkspaceEditorModule&
 						}
 					});
 			}),
-			Workspace::WorkspaceTabs::MiddleDocumentArea);
+			Workspace::WorkspaceTabs::TopMiddleDocumentArea);
 	AnimNextGraphDocumentArgs.OnGetTabName = Workspace::FOnGetTabName::CreateLambda([](const Workspace::FWorkspaceEditorContext& InContext)
 	{
 		const UAnimNextGraph* Graph = CastChecked<UAnimNextGraph>(InContext.Object);
@@ -405,7 +414,7 @@ void FModule::RegisterWorkspaceDocumentTypes(Workspace::IWorkspaceEditorModule&
 
 	// --- AnimNextGraph_EdGraph ---
 	Workspace::FGraphDocumentWidgetArgs GraphArgs;
-	GraphArgs.SpawnLocation = Workspace::WorkspaceTabs::MiddleDocumentArea;
+	GraphArgs.SpawnLocation = Workspace::WorkspaceTabs::TopMiddleDocumentArea;
 	GraphArgs.OnCreateActionMenu = Workspace::FOnCreateActionMenu::CreateLambda([](const Workspace::FWorkspaceEditorContext& InContext, UEdGraph* InGraph, const FVector2D& InNodePosition, const TArray<UEdGraphPin*>& InDraggedPins, bool bAutoExpand, SGraphEditor::FActionMenuClosed InOnMenuClosed)
 	{
 		TSharedRef<SActionMenu> ActionMenu = SNew(SActionMenu, InGraph)
diff --git a/Engine/Plugins/Experimental/Workspace/Source/WorkspaceEditor/Private/WorkspaceEditor.cpp b/Engine/Plugins/Experimental/Workspace/Source/WorkspaceEditor/Private/WorkspaceEditor.cpp
--- a/Engine/Plugins/Experimental/Workspace/Source/WorkspaceEditor/Private/WorkspaceEditor.cpp
+++ b/Engine/Plugins/Experimental/Workspace/Source/WorkspaceEditor/Private/WorkspaceEditor.cpp
@@ -15,6 +15,7 @@
 #include "Toolkits/AssetEditorToolkitMenuContext.h"
 #include "WorkspaceEditorCommands.h"
 #include "SWorkspaceTabWrapper.h"
+#include "Framework/Docking/LayoutExtender.h"
 
 #define LOCTEXT_NAMESPACE "WorkspaceEditor"
 
@@ -30,11 +31,12 @@ namespace WorkspaceTabs
 {
 	const FName Details("DetailsTab");
 	const FName WorkspaceView("WorkspaceView");
-	const FName TopDocumentArea("TopDocumentArea");
-	const FName LeftDocumentArea("LeftDocumentArea");
-	const FName MiddleDocumentArea("MiddleDocumentArea");
-	const FName RightDocumentArea("RightDocumentArea");
-	const FName BottomDocumentArea("BottomDocumentArea");
+	const FName TopLeftDocumentArea("TopLeftDocumentArea");
+	const FName BottomLeftDocumentArea("BottomLeftDocumentArea");
+	const FName TopMiddleDocumentArea("TopMiddleDocumentArea");
+	const FName BottomMiddleDocumentArea("BottomMiddleDocumentArea");
+	const FName TopRightDocumentArea("TopRightDocumentArea");
+	const FName BottomRightDocumentArea("BottomRightDocumentArea");
 }
 
 const FName WorkspaceAppIdentifier("WorkspaceEditor");
@@ -48,78 +50,117 @@ void FWorkspaceEditor::CreateWidgets()
 {
 	DocumentManager = MakeShared<FDocumentTracker>(NAME_None);
 	DocumentManager->Initialize(SharedThis(this));
-	
+
 	FBaseAssetToolkit::CreateWidgets();
 
-	const FWorkspaceEditorModule& WorkspaceEditorModule = FModuleManager::Get().LoadModuleChecked<FWorkspaceEditorModule>("WorkspaceEditor");
+	FWorkspaceEditorModule& WorkspaceEditorModule = FModuleManager::Get().LoadModuleChecked<FWorkspaceEditorModule>("WorkspaceEditor");
 
 	// Build document summoners for each workspace layout area
-	const TSharedRef<FAssetDocumentSummoner> LeftAssetDocumentSummoner = MakeShared<FAssetDocumentSummoner>(WorkspaceTabs::LeftDocumentArea, SharedThis(this));
-	LeftAssetDocumentSummoner->SetAllowedClassPaths(WorkspaceEditorModule.GetAllowedObjectTypesForArea(WorkspaceTabs::LeftDocumentArea));
-	DocumentManager->RegisterDocumentFactory(LeftAssetDocumentSummoner);
+	const TSharedRef<FAssetDocumentSummoner> TopLeftAssetDocumentSummoner = MakeShared<FAssetDocumentSummoner>(WorkspaceTabs::TopLeftDocumentArea, SharedThis(this));
+	TopLeftAssetDocumentSummoner->SetAllowedClassPaths(WorkspaceEditorModule.GetAllowedObjectTypesForArea(WorkspaceTabs::TopLeftDocumentArea));
+	DocumentManager->RegisterDocumentFactory(TopLeftAssetDocumentSummoner);
+
+	const TSharedRef<FAssetDocumentSummoner> BottomLeftAssetDocumentSummoner = MakeShared<FAssetDocumentSummoner>(WorkspaceTabs::BottomLeftDocumentArea, SharedThis(this));
+	BottomLeftAssetDocumentSummoner->SetAllowedClassPaths(WorkspaceEditorModule.GetAllowedObjectTypesForArea(WorkspaceTabs::BottomLeftDocumentArea));
+	DocumentManager->RegisterDocumentFactory(BottomLeftAssetDocumentSummoner);
 
 	constexpr bool bAllowUnsupportedClasses = true;
-	const TSharedRef<FAssetDocumentSummoner> MiddleAssetDocumentSummoner = MakeShared<FAssetDocumentSummoner>(WorkspaceTabs::MiddleDocumentArea, SharedThis(this), bAllowUnsupportedClasses);
-	MiddleAssetDocumentSummoner->SetAllowedClassPaths(WorkspaceEditorModule.GetAllowedObjectTypesForArea(WorkspaceTabs::MiddleDocumentArea));
-	DocumentManager->RegisterDocumentFactory(MiddleAssetDocumentSummoner);
+	const TSharedRef<FAssetDocumentSummoner> TopMiddleAssetDocumentSummoner = MakeShared<FAssetDocumentSummoner>(WorkspaceTabs::TopMiddleDocumentArea, SharedThis(this), bAllowUnsupportedClasses);
+	TopMiddleAssetDocumentSummoner->SetAllowedClassPaths(WorkspaceEditorModule.GetAllowedObjectTypesForArea(WorkspaceTabs::TopMiddleDocumentArea));
+	DocumentManager->RegisterDocumentFactory(TopMiddleAssetDocumentSummoner);
+
+	const TSharedRef<FAssetDocumentSummoner> BottomMiddleAssetDocumentSummoner = MakeShared<FAssetDocumentSummoner>(WorkspaceTabs::BottomMiddleDocumentArea, SharedThis(this));
+	BottomMiddleAssetDocumentSummoner->SetAllowedClassPaths(WorkspaceEditorModule.GetAllowedObjectTypesForArea(WorkspaceTabs::BottomMiddleDocumentArea));
+	DocumentManager->RegisterDocumentFactory(BottomMiddleAssetDocumentSummoner);
+
+	const TSharedRef<FAssetDocumentSummoner> TopRightAssetDocumentSummoner = MakeShared<FAssetDocumentSummoner>(WorkspaceTabs::TopRightDocumentArea, SharedThis(this));
+	TopRightAssetDocumentSummoner->SetAllowedClassPaths(WorkspaceEditorModule.GetAllowedObjectTypesForArea(WorkspaceTabs::TopRightDocumentArea));
+	DocumentManager->RegisterDocumentFactory(TopRightAssetDocumentSummoner);
+
+	const TSharedRef<FAssetDocumentSummoner> BottomRightAssetDocumentSummoner = MakeShared<FAssetDocumentSummoner>(WorkspaceTabs::BottomRightDocumentArea, SharedThis(this));
+	BottomRightAssetDocumentSummoner->SetAllowedClassPaths(WorkspaceEditorModule.GetAllowedObjectTypesForArea(WorkspaceTabs::BottomRightDocumentArea));
+	DocumentManager->RegisterDocumentFactory(BottomRightAssetDocumentSummoner);
 
-	const TSharedRef<FAssetDocumentSummoner> RightAssetDocumentSummoner = MakeShared<FAssetDocumentSummoner>(WorkspaceTabs::RightDocumentArea, SharedThis(this));
-	RightAssetDocumentSummoner->SetAllowedClassPaths(WorkspaceEditorModule.GetAllowedObjectTypesForArea(WorkspaceTabs::RightDocumentArea));
-	DocumentManager->RegisterDocumentFactory(RightAssetDocumentSummoner);
-	
 	check(DetailsView.IsValid());
 	WorkspaceEditorModule.ApplyWorkspaceDetailsCustomization(StaticCastWeakPtr<IWorkspaceEditor>(this->AsWeak()), DetailsView);
 
-	StandaloneDefaultLayout = FTabManager::NewLayout("Standalone_WorkspaceEditor_Layout_v1.1")
+	StandaloneDefaultLayout = FTabManager::NewLayout("Standalone_WorkspaceEditor_Layout_v1.2")
 	->AddArea
 	(
 		FTabManager::NewPrimaryArea()
 		->SetOrientation(Orient_Vertical)
 		->Split
-		(
-			FTabManager::NewStack()
-			->SetSizeCoefficient(0.15f)
-			->SetHideTabWell(false)
-			->AddTab(WorkspaceTabs::TopDocumentArea, ETabState::ClosedTab)
-		)
-		->Split
 		(
 			FTabManager::NewSplitter()
-			->SetSizeCoefficient(0.6f)
 			->SetOrientation(Orient_Horizontal)
 			->Split
 			(
-				FTabManager::NewStack()
+				FTabManager::NewSplitter()
 				->SetSizeCoefficient(0.25f)
-				->SetHideTabWell(false)
-				->AddTab(WorkspaceTabs::WorkspaceView, ETabState::OpenedTab)
-				->AddTab(WorkspaceTabs::LeftDocumentArea, ETabState::ClosedTab)
+				->SetOrientation(Orient_Vertical)
+				->Split
+				(
+					FTabManager::NewStack()
+					->SetSizeCoefficient(0.5f)
+					->SetHideTabWell(false)
+					->AddTab(WorkspaceTabs::WorkspaceView, ETabState::OpenedTab)
+					->AddTab(WorkspaceTabs::TopLeftDocumentArea, ETabState::ClosedTab)
+				)
+				->Split
+				(
+					FTabManager::NewStack()
+					->SetSizeCoefficient(0.5f)
+					->SetHideTabWell(false)
+					->AddTab(WorkspaceTabs::BottomLeftDocumentArea, ETabState::ClosedTab)
+				)
 			)
 			->Split
 			(
-				FTabManager::NewStack()
+				FTabManager::NewSplitter()
 				->SetSizeCoefficient(0.5f)
-				->SetHideTabWell(false)
-				->AddTab(WorkspaceTabs::MiddleDocumentArea, ETabState::ClosedTab)
+				->SetOrientation(Orient_Vertical)
+				->Split
+				(
+					FTabManager::NewStack()
+					->SetSizeCoefficient(0.75f)
+					->SetHideTabWell(false)
+					->AddTab(WorkspaceTabs::TopMiddleDocumentArea, ETabState::ClosedTab)
+				)
+				->Split
+				(
+					FTabManager::NewStack()
+					->SetSizeCoefficient(0.25f)
+					->SetHideTabWell(false)
+					->AddTab(WorkspaceTabs::BottomMiddleDocumentArea, ETabState::ClosedTab)
+				)
 			)
 			->Split
 			(
-				FTabManager::NewStack()
+				FTabManager::NewSplitter()
 				->SetSizeCoefficient(0.25f)
-				->SetHideTabWell(false)
-				->AddTab(WorkspaceTabs::RightDocumentArea, ETabState::ClosedTab)
-				->AddTab(FBaseAssetToolkit::DetailsTabID, ETabState::OpenedTab)
+				->SetOrientation(Orient_Vertical)
+				->Split
+				(
+					FTabManager::NewStack()
+					->SetSizeCoefficient(0.5f)
+					->SetHideTabWell(false)
+					->AddTab(WorkspaceTabs::TopRightDocumentArea, ETabState::ClosedTab)
+					->AddTab(FBaseAssetToolkit::DetailsTabID, ETabState::OpenedTab)
+				)
+				->Split
+				(
+					FTabManager::NewStack()
+					->SetSizeCoefficient(0.5f)
+					->SetHideTabWell(false)
+					->AddTab(WorkspaceTabs::BottomRightDocumentArea, ETabState::ClosedTab)
+				)
 			)
 		)
-		->Split
-		(
-			FTabManager::NewStack()
-			->SetSizeCoefficient(0.15f)
-			->SetHideTabWell(false)
-			->AddTab(WorkspaceTabs::BottomDocumentArea, ETabState::ClosedTab)
-		)
 	);
 
+	WorkspaceEditorModule.OnExtendTabs().Broadcast(*LayoutExtender.Get(), StaticCastSharedRef<UE::Workspace::IWorkspaceEditor>(AsShared()));
+	StandaloneDefaultLayout->ProcessExtensions(*LayoutExtender.Get());
+
 	WorkspaceView = SNew(SWorkspaceView, Workspace, StaticCastSharedRef<UE::Workspace::IWorkspaceEditor>(AsShared()));
 	BindCommands();
 }
diff --git a/Engine/Plugins/Experimental/Workspace/Source/WorkspaceEditor/Private/WorkspaceEditorModule.h b/Engine/Plugins/Experimental/Workspace/Source/WorkspaceEditor/Private/WorkspaceEditorModule.h
--- a/Engine/Plugins/Experimental/Workspace/Source/WorkspaceEditor/Private/WorkspaceEditorModule.h
+++ b/Engine/Plugins/Experimental/Workspace/Source/WorkspaceEditor/Private/WorkspaceEditorModule.h
@@ -32,6 +32,7 @@ private:
 	virtual void RegisterWorkspaceItemDetails(const FOutlinerItemDetailsId& InItemDetailsId, TSharedPtr<IWorkspaceOutlinerItemDetails> InItemDetails) override;	
 	virtual void UnregisterWorkspaceItemDetails(const FOutlinerItemDetailsId& InItemDetails) override;
 	virtual FOnRegisterTabs& OnRegisterTabsForEditor() override { return RegisterTabsForEditor; }
+	virtual FOnExtendTabs& OnExtendTabs() override { return ExtendTabsForEditor; }
 
 	// Find an existing registered object document type
 	const FObjectDocumentArgs* FindObjectDocumentType(const FTopLevelAssetPath& InClassPath) const;
@@ -57,6 +58,9 @@ private:
 
 	/** Event called to allow external clients to register additional tabs for the specified editor */
 	FOnRegisterTabs RegisterTabsForEditor;
+	
+	/** Event called to allow external clients to extend the tab layout for the specified editor */
+	FOnExtendTabs ExtendTabsForEditor;
 
 	friend struct FAssetDocumentSummoner;
 	friend struct FWorkspaceOutlinerTreeItem;
diff --git a/Engine/Plugins/Experimental/Workspace/Source/WorkspaceEditor/Public/IWorkspaceEditorModule.h b/Engine/Plugins/Experimental/Workspace/Source/WorkspaceEditor/Public/IWorkspaceEditorModule.h
--- a/Engine/Plugins/Experimental/Workspace/Source/WorkspaceEditor/Public/IWorkspaceEditorModule.h
+++ b/Engine/Plugins/Experimental/Workspace/Source/WorkspaceEditor/Public/IWorkspaceEditorModule.h
@@ -16,6 +16,7 @@ class IDetailsView;
 class FTabManager;
 class FWorkflowAllowedTabSet;
 class IDetailCustomization;
+class FLayoutExtender;
 
 struct FSlateBrush;
 struct FTopLevelAssetPath;
@@ -32,9 +33,12 @@ namespace UE::Workspace
 
 namespace WorkspaceTabs
 {
-	WORKSPACEEDITOR_API extern const FName LeftDocumentArea;
-	WORKSPACEEDITOR_API extern const FName MiddleDocumentArea;
-	WORKSPACEEDITOR_API extern const FName RightDocumentArea;
+	WORKSPACEEDITOR_API extern const FName TopLeftDocumentArea;
+	WORKSPACEEDITOR_API extern const FName BottomLeftDocumentArea;
+	WORKSPACEEDITOR_API extern const FName TopMiddleDocumentArea;
+	WORKSPACEEDITOR_API extern const FName BottomMiddleDocumentArea;
+	WORKSPACEEDITOR_API extern const FName TopRightDocumentArea;
+	WORKSPACEEDITOR_API extern const FName BottomRightDocumentArea;
 }
 
 // Context passed to workspace editor delegates
@@ -78,7 +82,7 @@ struct FObjectDocumentArgs
 {
 	FObjectDocumentArgs() = default;
 
-	FObjectDocumentArgs(FOnMakeDocumentWidget InOnMakeDocumentWidget, FName InSpawnLocation = WorkspaceTabs::MiddleDocumentArea)
+	FObjectDocumentArgs(FOnMakeDocumentWidget InOnMakeDocumentWidget, FName InSpawnLocation = WorkspaceTabs::TopMiddleDocumentArea)
 		: OnMakeDocumentWidget(InOnMakeDocumentWidget)
 		, SpawnLocation(InSpawnLocation)
 	{}
@@ -99,7 +103,7 @@ struct FObjectDocumentArgs
 	FOnGetTabName OnGetTabName;
 
 	// Where to spawn the widget in the workspace layout - e.g. one of WorkspaceTabs
-	FName SpawnLocation = WorkspaceTabs::MiddleDocumentArea;
+	FName SpawnLocation = WorkspaceTabs::TopMiddleDocumentArea;
 
 	// Delegate called to get the bread crumb trail for this document tab
 	FOnGetDocumentBreadcrumbTrail OnGetDocumentBreadcrumbTrail;
@@ -127,7 +131,7 @@ using FOnGetWorkspaceDetailCustomizationInstance = TDelegate<TSharedRef<IDetailC
 struct FGraphDocumentWidgetArgs
 {
 	// Where to spawn the widget in the workspace layout - e.g. one of WorkspaceTabs
-	FName SpawnLocation = WorkspaceTabs::MiddleDocumentArea;
+	FName SpawnLocation = WorkspaceTabs::TopMiddleDocumentArea;
 
 	FOnCreateActionMenu OnCreateActionMenu;
 
@@ -191,6 +195,10 @@ public:
 	// Event to allow registering tabs to other elements
 	DECLARE_EVENT_ThreeParams(IWorkspaceEditorModule, FOnRegisterTabs, FWorkflowAllowedTabSet& TabFactories, const TSharedRef<FTabManager>&, TSharedPtr<IWorkspaceEditor>);
 	virtual FOnRegisterTabs& OnRegisterTabsForEditor() = 0;
+
+	// Event to allow extending the layout
+	DECLARE_EVENT_TwoParams(IWorkspaceEditorModule, FOnExtendTabs, FLayoutExtender&, TSharedPtr<IWorkspaceEditor>);
+	virtual FOnExtendTabs& OnExtendTabs() = 0;
 };
 
 }
