- Mark workspace entry package as dirty when removing them (this ensure they have the opportunity to get saved/get removed from source-control)
- Do _not_ mark workspace package itself dirty, to remove asset  'contention' from workflow

diff --git a/Engine/Plugins/Experimental/Workspace/Source/WorkspaceEditor/Private/Workspace.cpp b/Engine/Plugins/Experimental/Workspace/Source/WorkspaceEditor/Private/Workspace.cpp
--- a/Engine/Plugins/Experimental/Workspace/Source/WorkspaceEditor/Private/Workspace.cpp
+++ b/Engine/Plugins/Experimental/Workspace/Source/WorkspaceEditor/Private/Workspace.cpp
@@ -23,7 +23,7 @@ bool UWorkspace::AddAsset(const FAssetData& InAsset, bool bSetupUndoRedo, bool b
 
 	if(bSetupUndoRedo)
 	{
-		Modify();
+		Modify(false);
 	}
 
 	int32 NewIndex = INDEX_NONE; 
@@ -38,7 +38,7 @@ bool UWorkspace::AddAsset(const FAssetData& InAsset, bool bSetupUndoRedo, bool b
 		if(!Asset->HasAnyFlags(RF_Transient))
 		{
 			UWorkspaceAssetEntry* NewEntry = NewObject<UWorkspaceAssetEntry>(this, UWorkspaceAssetEntry::StaticClass(), NAME_None, RF_Transactional);
-			FExternalPackageHelper::SetPackagingMode(NewEntry, this, true, true, PKG_None);
+			FExternalPackageHelper::SetPackagingMode(NewEntry, this, true, false, PKG_None);
 
 			NewEntry->Asset = TSoftObjectPtr<UObject>(InAsset.GetSoftObjectPath());
 			NewIndex = AssetEntries.Add(NewEntry);
@@ -126,7 +126,7 @@ bool UWorkspace::RemoveAsset(const FAssetData& InAsset, bool bSetupUndoRedo, boo
 
 	if(bSetupUndoRedo)
 	{
-		Modify();
+		Modify(false);
 	}
 
 	const int32 EntryIndex = AssetEntries.IndexOfByPredicate([&InAsset](const UWorkspaceAssetEntry* AssetEntry) -> bool
@@ -139,6 +139,7 @@ bool UWorkspace::RemoveAsset(const FAssetData& InAsset, bool bSetupUndoRedo, boo
 		UWorkspaceAssetEntry* EntryToRemove = AssetEntries[EntryIndex];		
 		check(AssetEntries.Remove(EntryToRemove) == 1);
 		EntryToRemove->MarkAsGarbage();
+		EntryToRemove->MarkPackageDirty();
 		
 		BroadcastModified();
 	}
