Clear tab-factories on WorkspaceEditor shutdown as this introduced a self-shared ptr through the tab summoners

diff --git a/Engine/Plugins/Experimental/Workspace/Source/WorkspaceEditor/Private/WorkspaceEditor.cpp b/Engine/Plugins/Experimental/Workspace/Source/WorkspaceEditor/Private/WorkspaceEditor.cpp
--- a/Engine/Plugins/Experimental/Workspace/Source/WorkspaceEditor/Private/WorkspaceEditor.cpp
+++ b/Engine/Plugins/Experimental/Workspace/Source/WorkspaceEditor/Private/WorkspaceEditor.cpp
@@ -480,6 +480,7 @@ void FWorkspaceEditor::OnClose()
 		DetailsView.Reset();
 	}
 
+	TabFactories.Clear();
 	FBaseAssetToolkit::OnClose();
 }
 
