Added overlay for empty workspaces, guiding user to workflow to add assets
Added 'Add' button to outliner as an alternative to the asset-drag-and-drop

diff --git a/Engine/Plugins/Experimental/Workspace/Source/WorkspaceEditor/Private/Outliner/WorkspaceOutlinerMode.h b/Engine/Plugins/Experimental/Workspace/Source/WorkspaceEditor/Private/Outliner/WorkspaceOutlinerMode.h
--- a/Engine/Plugins/Experimental/Workspace/Source/WorkspaceEditor/Private/Outliner/WorkspaceOutlinerMode.h
+++ b/Engine/Plugins/Experimental/Workspace/Source/WorkspaceEditor/Private/Outliner/WorkspaceOutlinerMode.h
@@ -23,8 +23,8 @@ namespace UE::Workspace
 		virtual void OnItemClicked(FSceneOutlinerTreeItemPtr Item) override;
 		virtual FReply OnKeyDown(const FKeyEvent& InKeyEvent) override;
 		void HandleItemSelection(const FSceneOutlinerItemSelection& Selection);
-
 		virtual void OnItemSelectionChanged(FSceneOutlinerTreeItemPtr Item, ESelectInfo::Type SelectionType, const FSceneOutlinerItemSelection& Selection) override;
+		virtual bool CanCustomizeToolbar() const { return true; }
 	protected:
 		virtual TUniquePtr<ISceneOutlinerHierarchy> CreateHierarchy() override;		
 		// End ISceneOutlinerMode overrides
diff --git a/Engine/Plugins/Experimental/Workspace/Source/WorkspaceEditor/Private/SWorkspaceView.cpp b/Engine/Plugins/Experimental/Workspace/Source/WorkspaceEditor/Private/SWorkspaceView.cpp
--- a/Engine/Plugins/Experimental/Workspace/Source/WorkspaceEditor/Private/SWorkspaceView.cpp
+++ b/Engine/Plugins/Experimental/Workspace/Source/WorkspaceEditor/Private/SWorkspaceView.cpp
@@ -7,6 +7,7 @@
 #include "SAssetDropTarget.h"
 #include "SceneOutlinerPublicTypes.h"
 #include "ScopedTransaction.h"
+#include "SPositiveActionButton.h"
 #include "SSceneOutliner.h"
 #include "WorkspaceSchema.h"
 #include "Framework/Commands/GenericCommands.h"
@@ -18,6 +19,76 @@
 namespace UE::Workspace
 {
 
+class SWorkspaceOutliner : public SSceneOutliner
+{
+public:
+	void Construct(const FArguments& InArgs, const FSceneOutlinerInitializationOptions& InitOptions, UWorkspace* InWorkspace)
+	{
+		AddAssetButton = SNew(SPositiveActionButton)
+				.OnGetMenuContent(this, &SWorkspaceOutliner::GetAddAssetPicker)			 
+				.Icon(FAppStyle::Get().GetBrush("Icons.Plus"))
+				.Text(LOCTEXT("AddAssetButton", "Add"));
+
+		WeakWorkspace = InWorkspace;
+		
+		SSceneOutliner::Construct(InArgs, InitOptions);
+	}
+
+	SWorkspaceOutliner() = default;
+	virtual ~SWorkspaceOutliner() = default;
+
+	void CustomAddToToolbar(TSharedPtr<class SHorizontalBox> Toolbar) override
+	{
+		Toolbar->AddSlot()
+		.VAlign(VAlign_Center)
+		.AutoWidth()
+		.Padding(4.f, 0.f, 0.f, 0.f)
+		[
+			AddAssetButton.ToSharedRef()
+		];
+	}
+
+	TSharedRef<SWidget> GetAddAssetPicker() const
+	{
+		if (UWorkspace* Workspace = WeakWorkspace.Get())
+		{
+			FContentBrowserModule& ContentBrowserModule = FModuleManager::Get().LoadModuleChecked<FContentBrowserModule>(TEXT("ContentBrowser"));
+
+			FAssetPickerConfig AssetPickerConfig;
+			AssetPickerConfig.InitialAssetViewType = EAssetViewType::List;
+			AssetPickerConfig.OnAssetSelected = FOnAssetSelected::CreateLambda([this, Workspace](const FAssetData& InAssetData)
+			{
+				FScopedTransaction Transaction(LOCTEXT("AddAsset", "Add asset to workspace"));
+				Workspace->AddAsset(InAssetData);
+		
+				AddAssetButton->SetIsMenuOpen(false, false);
+			});
+
+			TArray<FAssetData> WorkspaceAssetDataEntries;
+			Workspace->GetAssetDataEntries(WorkspaceAssetDataEntries);
+	
+			AssetPickerConfig.OnShouldFilterAsset = FOnShouldFilterAsset::CreateLambda([this, WorkspaceAssetDataEntries, Workspace](const FAssetData& InAssetData)
+			{
+				if (WorkspaceAssetDataEntries.Contains(InAssetData))
+				{
+					return true;
+				}
+	
+				return !Workspace->IsAssetSupported(InAssetData);
+			});
+
+			return ContentBrowserModule.Get().CreateAssetPicker(AssetPickerConfig);
+		}
+
+		return SNullWidget::NullWidget;	
+	
+	}
+private:
+	
+	TWeakObjectPtr<UWorkspace> WeakWorkspace;
+	TSharedPtr<SPositiveActionButton> AddAssetButton;
+};
+
 void SWorkspaceView::Construct(const FArguments& InArgs, UWorkspace* InWorkspace, TSharedRef<UE::Workspace::IWorkspaceEditor> InWorkspaceEditor)
 {
 	Workspace = InWorkspace;
@@ -31,8 +102,9 @@ void SWorkspaceView::Construct(const FArguments& InArgs, UWorkspace* InWorkspace
 		InitOptions.ColumnMap.Add(FWorkspaceOutlinerSourceControlColumn::GetID(), FSceneOutlinerColumnInfo(ESceneOutlinerColumnVisibility::Visible, 100, FCreateSceneOutlinerColumn::CreateLambda([](ISceneOutliner& InSceneOutliner) { return MakeShareable(new FWorkspaceOutlinerSourceControlColumn(InSceneOutliner)); }), false));	
 		InitOptions.ModeFactory = FCreateSceneOutlinerMode::CreateLambda([this, WeakWorkspaceEditor=InWorkspaceEditor.ToWeakPtr()](SSceneOutliner* InOutliner) { return new UE::Workspace::FWorkspaceOutlinerMode(UE::Workspace::FWorkspaceOutlinerMode(InOutliner, Workspace, WeakWorkspaceEditor)); });
 	}
-	SceneWorkspaceOutliner = SNew(SSceneOutliner, InitOptions);
-
+	SceneWorkspaceOutliner = SNew(SWorkspaceOutliner, InitOptions, Workspace);
+	
+	TWeakObjectPtr<UWorkspace> WeakWorkspace = Workspace;
 	ChildSlot
 	[
 		SNew(SAssetDropTarget)
@@ -40,7 +112,6 @@ void SWorkspaceView::Construct(const FArguments& InArgs, UWorkspace* InWorkspace
 		.OnAssetsDropped_Lambda([this](const FDragDropEvent& InEvent, TArrayView<FAssetData> InAssets)
 		{
 			FScopedTransaction Transaction(LOCTEXT("AddAssets", "Add assets to workspace"));
-
 			Workspace->AddAssets(InAssets);
 		})
 		.OnAreAssetsAcceptableForDropWithReason_Lambda([this](TArrayView<FAssetData> InAssets, FText& OutText)
@@ -58,7 +129,39 @@ void SWorkspaceView::Construct(const FArguments& InArgs, UWorkspace* InWorkspace
 		})
 		.Content()
 		[
-			SceneWorkspaceOutliner.ToSharedRef()
+			SNew(SOverlay)
+			+SOverlay::Slot()
+			[
+				SceneWorkspaceOutliner.ToSharedRef()
+			]
+			+SOverlay::Slot()
+			.Padding(32.f)
+			.VAlign(VAlign_Center)
+			.HAlign(HAlign_Center)
+			[
+				SNew(STextBlock)
+				.Text(
+					LOCTEXT("EmptyWorkspaceTooltip", "No assets currently in this workspace. Use 'Add' button or drag and drop to add assets.")
+				)
+				.Justification(ETextJustify::Center)
+				.AutoWrapText(true)
+				.Visibility(
+					TAttribute<EVisibility>::Create(
+						TAttribute<EVisibility>::FGetter::CreateLambda([WeakWorkspace]()
+						{
+							if (UWorkspace* PinnedWorkspace = WeakWorkspace.Get())
+							{
+								if (!PinnedWorkspace->HasValidEntries())
+								{
+									return EVisibility::Visible;
+								}
+							}
+
+							return EVisibility::Collapsed;
+						})
+					)
+				)
+			]
 		]
 	];
 }
diff --git a/Engine/Plugins/Experimental/Workspace/Source/WorkspaceEditor/Private/SWorkspaceView.h b/Engine/Plugins/Experimental/Workspace/Source/WorkspaceEditor/Private/SWorkspaceView.h
--- a/Engine/Plugins/Experimental/Workspace/Source/WorkspaceEditor/Private/SWorkspaceView.h
+++ b/Engine/Plugins/Experimental/Workspace/Source/WorkspaceEditor/Private/SWorkspaceView.h
@@ -7,10 +7,10 @@
 
 class UWorkspace;
 class UWorkSpaceAssetUserData;
-class SSceneOutliner;
 
 namespace UE::Workspace
 {
+class SWorkspaceOutliner;
 class IWorkspaceEditor;
 
 class SWorkspaceView : public SCompoundWidget
@@ -21,7 +21,7 @@ class SWorkspaceView : public SCompoundWidget
 	void Construct(const FArguments& InArgs, UWorkspace* InWorkspace, TSharedRef<UE::Workspace::IWorkspaceEditor> InWorkspaceEditor);
 private:
 	UWorkspace* Workspace = nullptr;
-	TSharedPtr<SSceneOutliner> SceneWorkspaceOutliner;
+	TSharedPtr<SWorkspaceOutliner> SceneWorkspaceOutliner;
 };
 
 };
\ No newline at end of file
diff --git a/Engine/Plugins/Experimental/Workspace/Source/WorkspaceEditor/Private/Workspace.cpp b/Engine/Plugins/Experimental/Workspace/Source/WorkspaceEditor/Private/Workspace.cpp
--- a/Engine/Plugins/Experimental/Workspace/Source/WorkspaceEditor/Private/Workspace.cpp
+++ b/Engine/Plugins/Experimental/Workspace/Source/WorkspaceEditor/Private/Workspace.cpp
@@ -270,6 +270,19 @@ void UWorkspace::GetAssetDataEntries(TArray<FAssetData>& OutAssetDataEntries) co
 	AssetRegistry.GetAssets(Filter, OutAssetDataEntries);
 }
 
+bool UWorkspace::HasValidEntries() const
+{
+	for(const UWorkspaceAssetEntry* AssetEntry : AssetEntries)
+	{
+		if (AssetEntry && AssetEntry->Asset.IsValid())
+		{
+			return true;
+		}
+	}
+
+	return false;
+}
+
 void UWorkspace::ReportError(const TCHAR* InMessage) const
 {
 #if WITH_EDITOR
diff --git a/Engine/Plugins/Experimental/Workspace/Source/WorkspaceEditor/Private/Workspace.h b/Engine/Plugins/Experimental/Workspace/Source/WorkspaceEditor/Private/Workspace.h
--- a/Engine/Plugins/Experimental/Workspace/Source/WorkspaceEditor/Private/Workspace.h
+++ b/Engine/Plugins/Experimental/Workspace/Source/WorkspaceEditor/Private/Workspace.h
@@ -22,6 +22,7 @@ namespace UE::Workspace
 	class SWorkspaceView;
 	class FWorkspaceOutlinerHierarchy;
 	class FWorkspaceOutlinerMode;
+	class SWorkspaceOutliner;
 
 	struct FWorkspaceOutliner;
 }
@@ -69,6 +70,7 @@ class UWorkspace : public UObject
 	friend struct UE::Workspace::FWorkspaceOutliner;
 	friend class UE::Workspace::FWorkspaceOutlinerHierarchy;
 	friend class UE::Workspace::FWorkspaceOutlinerMode;
+	friend class UE::Workspace::SWorkspaceOutliner;
 	
 	friend class UWorkspaceFactory;
 	friend class UAssetDefinition_Workspace;
@@ -115,6 +117,9 @@ class UWorkspace : public UObject
 	// Returns all contained AssetData entries
 	void GetAssetDataEntries(TArray<FAssetData>& OutAssetDataEntries) const;
 
+	// Returns whether or not the workspace contains _any_ valid UWorkspaceAssetEntry
+	bool HasValidEntries() const;
+
 	// UObject interface
 	virtual void PostLoad() override;
 	virtual void PostDuplicate(bool bDuplicateForPIE) override;
