Renamed FSceneOutlinerActorSCCColumn to FSceneOutlinerSourceControlColumn and exposed publicly
FSceneOutlinerSourceControlColumn only generates a widget in case FSceneOutlinerTreeItemSCC points to a valid package
Added new virtual ::GetPackageName to ISceneOutlinerTreeItem and deprecated FSceneOutlinerHelpers::GetExternalPackageName and FSceneOutlinerHelpers::GetExternalPackage behaviour
Added overrides for all the existing (relevant) tree item types

Replaced WorkspaceOutlinerSourceControlColumnn with SceneOutliner one (uniformity)
Added source control operations to WorkspaceOutliner context menu

diff --git a/Engine/Plugins/Experimental/Workspace/Source/WorkspaceEditor/Private/Outliner/WorkspaceOutlinerColumns.cpp b/Engine/Plugins/Experimental/Workspace/Source/WorkspaceEditor/Private/Outliner/WorkspaceOutlinerColumns.cpp
--- a/Engine/Plugins/Experimental/Workspace/Source/WorkspaceEditor/Private/Outliner/WorkspaceOutlinerColumns.cpp
+++ b/Engine/Plugins/Experimental/Workspace/Source/WorkspaceEditor/Private/Outliner/WorkspaceOutlinerColumns.cpp
@@ -16,84 +16,6 @@
 
 namespace UE::Workspace
 {
-	FName WorkspaceOutlinerSourceControl("Source Control Status");
-	FName FWorkspaceOutlinerSourceControlColumn::GetID()
-	{
-		return WorkspaceOutlinerSourceControl;
-	}
-
-	SHeaderRow::FColumn::FArguments FWorkspaceOutlinerSourceControlColumn::ConstructHeaderRowColumn()
-	{
-		return SHeaderRow::Column(GetColumnID())
-		.FixedWidth(24.f)
-		.HAlignHeader(HAlign_Center)
-		.VAlignHeader(VAlign_Center)
-		.HAlignCell(HAlign_Center)
-		.VAlignCell(VAlign_Center)
-		[
-			SNew(SBox)
-			.WidthOverride(16.0f)
-			.HeightOverride(16.0f)
-			.VAlign(VAlign_Center)
-			.HAlign(HAlign_Center)
-			[
-				SNew(SImage)
-				.ColorAndOpacity(FSlateColor::UseForeground())
-				.Image(FRevisionControlStyleManager::Get().GetBrush("RevisionControl.Icon"))
-				.ToolTipText(LOCTEXT("RevisionControlStatusTooltip", "Revision control status of this parameter"))
-			]
-		];	
-	}
-
-	const TSharedRef<SWidget> FWorkspaceOutlinerSourceControlColumn::ConstructRowWidget(FSceneOutlinerTreeItemRef Item, const STableRow<FSceneOutlinerTreeItemPtr>& Row)
-	{
-		if (const FWorkspaceOutlinerTreeItem* TreeItem = Item->CastTo<FWorkspaceOutlinerTreeItem>())
-		{
-			FString PackagePath;
-			if (const TSharedPtr<IWorkspaceOutlinerItemDetails> SharedFactory = FWorkspaceEditorModule::GetOutlinerItemDetails(MakeOutlinerDetailsId(TreeItem->Export)))
-			{
-				PackagePath = SourceControlHelpers::PackageFilename(SharedFactory->GetPackage(TreeItem->Export));
-			}
-			else if (TreeItem->Export.ParentIdentifier == NAME_None && TreeItem->Export.AssetPath.IsValid())
-			{
-				PackagePath = SourceControlHelpers::PackageFilename(TreeItem->Export.AssetPath.GetLongPackageName());
-			}
-		
-			if (!PackagePath.IsEmpty())
-			{
-				return SNew(SBox)
-				.HAlign(HAlign_Center)
-				.VAlign(VAlign_Center)
-				.HeightOverride(20.0f)
-				[
-					SNew(SImage)
-					.ToolTipText_Lambda([this, PackagePath]() -> FText
-					{
-						ISourceControlProvider& SourceControlProvider = ISourceControlModule::Get().GetProvider();
-						if(const FSourceControlStatePtr State = SourceControlProvider.GetState(PackagePath, EStateCacheUsage::Use))
-						{
-							return FText::Format(LOCTEXT("RevisionControlStatusFormat", "File: {0}\nStatus: {1}"), State->GetDisplayName(), State->GetDisplayTooltip());
-						}
-
-						return LOCTEXT("RevisionControlStatus", "Revision control status of this entry");
-					})
-					.Image_Lambda([this, PackagePath]() -> const FSlateBrush*
-					{
-						ISourceControlProvider& SourceControlProvider = ISourceControlModule::Get().GetProvider();
-						if(const FSourceControlStatePtr State = SourceControlProvider.GetState(PackagePath, EStateCacheUsage::Use))
-						{
-							return State->GetIcon().GetSmallIcon();
-						}
-						
-						return nullptr;
-					})
-				];
-			}
-		}
-		
-		return SNullWidget::NullWidget;
-	}
-
 	FName WorkspaceOutlinerFileState("File State");
 	FName FWorkspaceOutlinerFileStateColumn::GetID()
 	{
diff --git a/Engine/Plugins/Experimental/Workspace/Source/WorkspaceEditor/Private/Outliner/WorkspaceOutlinerColumns.h b/Engine/Plugins/Experimental/Workspace/Source/WorkspaceEditor/Private/Outliner/WorkspaceOutlinerColumns.h
--- a/Engine/Plugins/Experimental/Workspace/Source/WorkspaceEditor/Private/Outliner/WorkspaceOutlinerColumns.h
+++ b/Engine/Plugins/Experimental/Workspace/Source/WorkspaceEditor/Private/Outliner/WorkspaceOutlinerColumns.h
@@ -13,25 +13,6 @@ class SWidget;
 
 namespace UE::Workspace
 {
-class FWorkspaceOutlinerSourceControlColumn : public ISceneOutlinerColumn
-{
-public:
-	static FName GetID();
-	
-	FWorkspaceOutlinerSourceControlColumn(ISceneOutliner& SceneOutliner) : WeakSceneOutliner(StaticCastSharedRef<ISceneOutliner>(SceneOutliner.AsShared())) {}
-	virtual ~FWorkspaceOutlinerSourceControlColumn() override = default;
-	
-	// Begin ISceneOutlinerColumn overrides
-	virtual FName GetColumnID() override { return GetID(); }
-	virtual SHeaderRow::FColumn::FArguments ConstructHeaderRowColumn() override;
-	virtual const TSharedRef<SWidget> ConstructRowWidget(FSceneOutlinerTreeItemRef TreeItem, const STableRow<FSceneOutlinerTreeItemPtr>& Row) override;
-	virtual bool SupportsSorting() const override { return false; }
-	// End ISceneOutlinerColumn overrides
-
-private:
-	TWeakPtr<ISceneOutliner> WeakSceneOutliner;
-};
-
 class FWorkspaceOutlinerFileStateColumn : public ISceneOutlinerColumn
 {
 public:
diff --git a/Engine/Plugins/Experimental/Workspace/Source/WorkspaceEditor/Private/Outliner/WorkspaceOutlinerMode.cpp b/Engine/Plugins/Experimental/Workspace/Source/WorkspaceEditor/Private/Outliner/WorkspaceOutlinerMode.cpp
--- a/Engine/Plugins/Experimental/Workspace/Source/WorkspaceEditor/Private/Outliner/WorkspaceOutlinerMode.cpp
+++ b/Engine/Plugins/Experimental/Workspace/Source/WorkspaceEditor/Private/Outliner/WorkspaceOutlinerMode.cpp
@@ -2,6 +2,7 @@
 
 #include "WorkspaceOutlinerMode.h"
 
+#include "ISourceControlModule.h"
 #include "WorkspaceItemMenuContext.h"
 #include "WorkspaceOutlinerHierarchy.h"
 #include "WorkspaceOutlinerTreeItem.h"
@@ -56,7 +57,8 @@ TSharedPtr<SWidget> FWorkspaceOutlinerMode::CreateContextMenu()
 		FToolMenuOwnerScoped ToolMenuOwnerScope(this);
 		if (UToolMenu* Menu = ToolMenus->RegisterMenu(MenuName))
 		{
-			Menu->AddDynamicSection(TEXT("Assets"), FNewToolMenuDelegate::CreateLambda([](UToolMenu* InMenu)
+			TWeakPtr<SSceneOutliner> WeakOutliner = StaticCastSharedRef<SSceneOutliner>(SceneOutliner->AsShared());
+			Menu->AddDynamicSection(TEXT("Assets"), FNewToolMenuDelegate::CreateLambda([WeakOutliner](UToolMenu* InMenu)
 			{
 				const UAssetEditorToolkitMenuContext* EditorContext = InMenu->FindContext<UAssetEditorToolkitMenuContext>();
 				const UWorkspaceItemMenuContext* MenuContext = InMenu->FindContext<UWorkspaceItemMenuContext>();
@@ -127,6 +129,11 @@ TSharedPtr<SWidget> FWorkspaceOutlinerMode::CreateContextMenu()
 						);
 					}
 				}
+
+				if (TSharedPtr<SSceneOutliner> SharedOutliner = WeakOutliner.Pin())
+				{
+					SharedOutliner->AddSourceControlMenuOptions(InMenu);
+				}
 			}));
 		}
 	}
diff --git a/Engine/Plugins/Experimental/Workspace/Source/WorkspaceEditor/Private/Outliner/WorkspaceOutlinerTreeItem.cpp b/Engine/Plugins/Experimental/Workspace/Source/WorkspaceEditor/Private/Outliner/WorkspaceOutlinerTreeItem.cpp
--- a/Engine/Plugins/Experimental/Workspace/Source/WorkspaceEditor/Private/Outliner/WorkspaceOutlinerTreeItem.cpp
+++ b/Engine/Plugins/Experimental/Workspace/Source/WorkspaceEditor/Private/Outliner/WorkspaceOutlinerTreeItem.cpp
@@ -114,6 +114,20 @@ namespace UE::Workspace
 	{
 		return SNew(SWorkspaceOutlinerTreelabel, *this, Outliner, InRow);
 	}
+
+	FString FWorkspaceOutlinerTreeItem::GetPackageName() const
+	{
+		if (const TSharedPtr<IWorkspaceOutlinerItemDetails> SharedFactory = FWorkspaceEditorModule::GetOutlinerItemDetails(MakeOutlinerDetailsId(Export)))
+		{
+			return SharedFactory->GetPackage(Export)->GetName();
+		}
+		else if (Export.ParentIdentifier == NAME_None && Export.AssetPath.IsValid())
+		{
+			return Export.AssetPath.GetLongPackageName();
+		}		
+		
+		return ISceneOutlinerTreeItem::GetPackageName();
+	}
 }
 
 #undef LOCTEXT_NAMESPACE // "WorkspaceOutlinerTreeItem"
\ No newline at end of file
diff --git a/Engine/Plugins/Experimental/Workspace/Source/WorkspaceEditor/Private/Outliner/WorkspaceOutlinerTreeItem.h b/Engine/Plugins/Experimental/Workspace/Source/WorkspaceEditor/Private/Outliner/WorkspaceOutlinerTreeItem.h
--- a/Engine/Plugins/Experimental/Workspace/Source/WorkspaceEditor/Private/Outliner/WorkspaceOutlinerTreeItem.h
+++ b/Engine/Plugins/Experimental/Workspace/Source/WorkspaceEditor/Private/Outliner/WorkspaceOutlinerTreeItem.h
@@ -25,6 +25,7 @@ namespace UE::Workspace
 		virtual FString GetDisplayString() const override;
 		virtual bool CanInteract() const override { return true; }
 		virtual TSharedRef<SWidget> GenerateLabelWidget(ISceneOutliner& Outliner, const STableRow<FSceneOutlinerTreeItemPtr>& InRow) override;
+		virtual FString GetPackageName() const override;
 		// End ISceneOutlinerTreeItem overrides
 
 		// AssetRegistry export this tree item represents in the Outliner 
diff --git a/Engine/Plugins/Experimental/Workspace/Source/WorkspaceEditor/Private/SWorkspaceView.cpp b/Engine/Plugins/Experimental/Workspace/Source/WorkspaceEditor/Private/SWorkspaceView.cpp
--- a/Engine/Plugins/Experimental/Workspace/Source/WorkspaceEditor/Private/SWorkspaceView.cpp
+++ b/Engine/Plugins/Experimental/Workspace/Source/WorkspaceEditor/Private/SWorkspaceView.cpp
@@ -6,6 +6,7 @@
 #include "IContentBrowserSingleton.h"
 #include "SAssetDropTarget.h"
 #include "SceneOutlinerPublicTypes.h"
+#include "SceneOutlinerSourceControlColumn.h"
 #include "ScopedTransaction.h"
 #include "SPositiveActionButton.h"
 #include "SSceneOutliner.h"
@@ -99,7 +100,7 @@ void SWorkspaceView::Construct(const FArguments& InArgs, UWorkspace* InWorkspace
 		InitOptions.bShowHeaderRow = true;
 		InitOptions.ColumnMap.Add(FSceneOutlinerBuiltInColumnTypes::Label(), FSceneOutlinerColumnInfo(ESceneOutlinerColumnVisibility::Visible, 10));
 		InitOptions.ColumnMap.Add(FWorkspaceOutlinerFileStateColumn::GetID(), FSceneOutlinerColumnInfo(ESceneOutlinerColumnVisibility::Visible, 0, FCreateSceneOutlinerColumn::CreateLambda([](ISceneOutliner& InSceneOutliner) { return MakeShareable(new FWorkspaceOutlinerFileStateColumn(InSceneOutliner)); }), false));
-		InitOptions.ColumnMap.Add(FWorkspaceOutlinerSourceControlColumn::GetID(), FSceneOutlinerColumnInfo(ESceneOutlinerColumnVisibility::Visible, 100, FCreateSceneOutlinerColumn::CreateLambda([](ISceneOutliner& InSceneOutliner) { return MakeShareable(new FWorkspaceOutlinerSourceControlColumn(InSceneOutliner)); }), false));	
+		InitOptions.ColumnMap.Add(FSceneOutlinerSourceControlColumn::GetID(), FSceneOutlinerColumnInfo(ESceneOutlinerColumnVisibility::Visible, 100, FCreateSceneOutlinerColumn::CreateLambda([](ISceneOutliner& InSceneOutliner) { return MakeShareable(new FSceneOutlinerSourceControlColumn(InSceneOutliner)); }), false));	
 		InitOptions.ModeFactory = FCreateSceneOutlinerMode::CreateLambda([this, WeakWorkspaceEditor=InWorkspaceEditor.ToWeakPtr()](SSceneOutliner* InOutliner) { return new UE::Workspace::FWorkspaceOutlinerMode(UE::Workspace::FWorkspaceOutlinerMode(InOutliner, Workspace, WeakWorkspaceEditor)); });
 	}
 	SceneWorkspaceOutliner = SNew(SWorkspaceOutliner, InitOptions, Workspace);
diff --git a/Engine/Source/Editor/LevelEditor/Private/LevelEditorOutlinerSettings.cpp b/Engine/Source/Editor/LevelEditor/Private/LevelEditorOutlinerSettings.cpp
--- a/Engine/Source/Editor/LevelEditor/Private/LevelEditorOutlinerSettings.cpp
+++ b/Engine/Source/Editor/LevelEditor/Private/LevelEditorOutlinerSettings.cpp
@@ -300,12 +300,12 @@ void FLevelEditorOutlinerSettings::CreateDefaultFilters()
 
 bool FLevelEditorOutlinerSettings::DoesActorPassUnsavedFilter(const ISceneOutlinerTreeItem& InItem)
 {
-	return UnsavedAssets.Contains(USourceControlHelpers::PackageFilename(SceneOutliner::FSceneOutlinerHelpers::GetExternalPackageName(InItem)));
+	return UnsavedAssets.Contains(USourceControlHelpers::PackageFilename(InItem.GetPackageName()));
 }
 
 bool FLevelEditorOutlinerSettings::DoesActorPassUncontrolledFilter(const ISceneOutlinerTreeItem& InItem)
 {
-	FString ExternalPackageFilename = USourceControlHelpers::PackageFilename(SceneOutliner::FSceneOutlinerHelpers::GetExternalPackageName(InItem));
+	FString ExternalPackageFilename = USourceControlHelpers::PackageFilename(InItem.GetPackageName());
 	
 	for (const TSharedRef<FUncontrolledChangelistState>& UncontrolledChangelistState : UncontrolledChangelistStates)
 	{
diff --git a/Engine/Source/Editor/SceneOutliner/Private/ActorBrowsingMode.cpp b/Engine/Source/Editor/SceneOutliner/Private/ActorBrowsingMode.cpp
--- a/Engine/Source/Editor/SceneOutliner/Private/ActorBrowsingMode.cpp
+++ b/Engine/Source/Editor/SceneOutliner/Private/ActorBrowsingMode.cpp
@@ -49,7 +49,6 @@
 #include "SourceControlOperations.h"
 #include "EditorLevelUtils.h"
 #include "EditorViewportCommands.h"
-#include "SceneOutlinerActorSCCColumn.h"
 #include "Misc/ScopedSlowTask.h"
 #include "Elements/Framework/EngineElementsLibrary.h"
 #include "Elements/Framework/TypedElementHandle.h"
diff --git a/Engine/Source/Editor/SceneOutliner/Private/ActorDescTreeItem.cpp b/Engine/Source/Editor/SceneOutliner/Private/ActorDescTreeItem.cpp
--- a/Engine/Source/Editor/SceneOutliner/Private/ActorDescTreeItem.cpp
+++ b/Engine/Source/Editor/SceneOutliner/Private/ActorDescTreeItem.cpp
@@ -344,6 +344,16 @@ bool FActorDescTreeItem::GetPinnedState() const
 	return false;
 }
 
+FString FActorDescTreeItem::GetPackageName() const
+{
+	if (const FWorldPartitionActorDescInstance* ActorDescInstance = *ActorDescHandle)
+	{
+		return ActorDescInstance->GetActorPackage().ToString();
+	}
+	
+	return IActorBaseTreeItem::GetPackageName();
+}
+
 UExternalDataLayerAsset* FActorDescTreeItem::GetExternalDataLayerAsset() const
 {
 	if (const FWorldPartitionActorDescInstance* ActorDescInstance = *ActorDescHandle)
diff --git a/Engine/Source/Editor/SceneOutliner/Private/ActorFolderTreeItem.cpp b/Engine/Source/Editor/SceneOutliner/Private/ActorFolderTreeItem.cpp
--- a/Engine/Source/Editor/SceneOutliner/Private/ActorFolderTreeItem.cpp
+++ b/Engine/Source/Editor/SceneOutliner/Private/ActorFolderTreeItem.cpp
@@ -398,6 +398,19 @@ FFolder FActorFolderTreeItem::GetFolder() const
 	return Folder.IsValid() ? Folder : FFolderTreeItem::GetFolder();
 }
 
+FString FActorFolderTreeItem::GetPackageName() const
+{
+	if (const UActorFolder* ActorFolderPtr = GetActorFolder())
+	{
+		if (ActorFolderPtr->IsPackageExternal())
+		{
+			return ActorFolderPtr->GetExternalPackage()->GetName();
+		}
+	}
+	
+	return FFolderTreeItem::GetPackageName();
+}
+
 bool FActorFolderTreeItem::CanInteract() const
 {
 	if (!FFolderTreeItem::CanInteract())
diff --git a/Engine/Source/Editor/SceneOutliner/Private/ActorTreeItem.cpp b/Engine/Source/Editor/SceneOutliner/Private/ActorTreeItem.cpp
--- a/Engine/Source/Editor/SceneOutliner/Private/ActorTreeItem.cpp
+++ b/Engine/Source/Editor/SceneOutliner/Private/ActorTreeItem.cpp
@@ -402,6 +402,19 @@ void FActorTreeItem::GenerateContextMenu(UToolMenu* Menu, SSceneOutliner& Outlin
 	}
 }
 
+FString FActorTreeItem::GetPackageName() const
+{
+	if (const AActor* ActorPtr = Actor.Get())
+	{
+		if (ActorPtr->IsPackageExternal())
+		{
+			return ActorPtr->GetExternalPackage()->GetName();
+		}
+	}
+	
+	return IActorBaseTreeItem::GetPackageName();
+}
+
 const FGuid& FActorTreeItem::GetGuid() const
 {
 	static const FGuid InvalidGuid;
diff --git a/Engine/Source/Editor/SceneOutliner/Private/SceneOutlinerActorSCCColumn.cpp b/Engine/Source/Editor/SceneOutliner/Private/SceneOutlinerActorSCCColumn.cpp
deleted file mode 100644
diff --git a/Engine/Source/Editor/SceneOutliner/Private/SceneOutlinerActorSCCColumn.h b/Engine/Source/Editor/SceneOutliner/Private/SceneOutlinerActorSCCColumn.h
deleted file mode 100644
diff --git a/Engine/Source/Editor/SceneOutliner/Private/SceneOutlinerModule.cpp b/Engine/Source/Editor/SceneOutliner/Private/SceneOutlinerModule.cpp
--- a/Engine/Source/Editor/SceneOutliner/Private/SceneOutlinerModule.cpp
+++ b/Engine/Source/Editor/SceneOutliner/Private/SceneOutlinerModule.cpp
@@ -10,7 +10,7 @@
 #include "SceneOutlinerActorInfoColumn.h"
 #include "SceneOutlinerGutter.h"
 #include "SceneOutlinerItemLabelColumn.h"
-#include "SceneOutlinerActorSCCColumn.h"
+#include "SceneOutlinerSourceControlColumn.h"
 #include "SceneOutlinerPinnedColumn.h"
 #include "SceneOutlinerTextInfoColumn.h"
 #include "SceneOutlinerUnsavedColumn.h"
@@ -56,7 +56,7 @@ void FSceneOutlinerModule::StartupModule()
 	// Register builtin column types which are not active by default
 	RegisterColumnType<FSceneOutlinerGutter>();
 	RegisterColumnType<FTypeInfoColumn>();
-	RegisterColumnType<FSceneOutlinerActorSCCColumn>();
+	RegisterColumnType<FSceneOutlinerSourceControlColumn>();
 	RegisterColumnType<FSceneOutlinerPinnedColumn>();
 	RegisterColumnType<FSceneOutlinerActorUnsavedColumn>();
 }
@@ -67,7 +67,7 @@ void FSceneOutlinerModule::ShutdownModule()
 	UnRegisterColumnType<FSceneOutlinerGutter>();
 	UnRegisterColumnType<FSceneOutlinerItemLabelColumn>();
 	UnRegisterColumnType<FTypeInfoColumn>();
-	UnRegisterColumnType<FSceneOutlinerActorSCCColumn>();
+	UnRegisterColumnType<FSceneOutlinerSourceControlColumn>();
 	UnRegisterColumnType<FSceneOutlinerPinnedColumn>();
 	UnRegisterColumnType<FSceneOutlinerActorUnsavedColumn>();
 }
diff --git a/Engine/Source/Editor/SceneOutliner/Private/SceneOutlinerSCCHandler.cpp b/Engine/Source/Editor/SceneOutliner/Private/SceneOutlinerSCCHandler.cpp
--- a/Engine/Source/Editor/SceneOutliner/Private/SceneOutlinerSCCHandler.cpp
+++ b/Engine/Source/Editor/SceneOutliner/Private/SceneOutlinerSCCHandler.cpp
@@ -145,7 +145,7 @@ void FSceneOutlinerSCCHandler::CacheCanExecuteVars()
 				continue;
 			}
 
-			if (!SourceControl->IsExternalPackage())
+			if (!SourceControl->HasValidPackage())
 			{
 				// this isn't an external package so we can't do anything with source control
 				bCanExecuteSCC = false;
@@ -356,7 +356,7 @@ void FSceneOutlinerSCCHandler::GetSelectedPackages(TArray<UPackage*>& OutPackage
 			continue;
 		}
 
-		UPackage* Package = SourceControl->GetPackage();
+		UPackage* Package = SourceControl->LoadPackage();
 		if (Package != nullptr)
 		{
 			OutPackages.Add(Package);
diff --git a/Engine/Source/Editor/SceneOutliner/Private/SceneOutlinerSourceControlColumn.cpp b/Engine/Source/Editor/SceneOutliner/Private/SceneOutlinerSourceControlColumn.cpp
new file mode 100644
--- /dev/null
+++ b/Engine/Source/Editor/SceneOutliner/Private/SceneOutlinerSourceControlColumn.cpp
@@ -0,0 +1,74 @@
+// Copyright Epic Games, Inc. All Rights Reserved.
+
+#include "SceneOutlinerSourceControlColumn.h"
+#include "Widgets/Images/SImage.h"
+#include "Widgets/Images/SLayeredImage.h"
+#include "Widgets/Views/STreeView.h"
+#include "ActorTreeItem.h"
+#include "ActorDescTreeItem.h"
+#include "ActorFolderTreeItem.h"
+#include "WorldPartition/WorldPartitionActorDesc.h"
+#include "ISourceControlModule.h"
+#include "SceneOutlinerTreeItemSCC.h"
+#include "SSourceControlWidget.h"
+#include "Misc/MessageDialog.h"
+#include "RevisionControlStyle/RevisionControlStyle.h"
+
+#define LOCTEXT_NAMESPACE "SceneOutlinerSourceControlColumn"
+
+FName FSceneOutlinerSourceControlColumn::GetColumnID()
+{
+	return GetID();
+}
+
+SHeaderRow::FColumn::FArguments FSceneOutlinerSourceControlColumn::ConstructHeaderRowColumn()
+{
+	TSharedRef<SLayeredImage> HeaderRowIcon = SNew(SLayeredImage)
+			.ColorAndOpacity(FSlateColor::UseForeground())
+			.Image(FRevisionControlStyleManager::Get().GetBrush("RevisionControl.Icon"));
+	
+	HeaderRowIcon->AddLayer(TAttribute<const FSlateBrush*>::CreateSP(this, &FSceneOutlinerSourceControlColumn::GetHeaderIconBadge));
+	
+	return SHeaderRow::Column(GetColumnID())
+		.FixedWidth(24.f)
+		.HAlignHeader(HAlign_Center)
+		.VAlignHeader(VAlign_Center)
+		.HAlignCell(HAlign_Center)
+		.VAlignCell(VAlign_Center)
+		.DefaultTooltip(FText::FromName(GetColumnID()))
+		[
+			HeaderRowIcon
+		];
+}
+
+const TSharedRef<SWidget> FSceneOutlinerSourceControlColumn::ConstructRowWidget(FSceneOutlinerTreeItemRef TreeItem, const STableRow<FSceneOutlinerTreeItemPtr>& Row)
+{
+	TSharedPtr<FSceneOutlinerTreeItemSCC> SourceControl = WeakSceneOutliner.Pin()->GetItemSourceControl(TreeItem);
+	if (SourceControl.IsValid() && SourceControl->HasValidPackage())
+	{
+		TSharedRef<SSourceControlWidget> Widget = SNew(SSourceControlWidget, SourceControl);
+		
+		return SNew(SHorizontalBox)
+			+ SHorizontalBox::Slot()
+			.AutoWidth()
+			.VAlign(VAlign_Center)
+			[
+				Widget
+			];
+	}
+	return SNullWidget::NullWidget;
+}
+
+const FSlateBrush* FSceneOutlinerSourceControlColumn::GetHeaderIconBadge() const
+{
+	if (ISourceControlModule::Get().IsEnabled())
+	{
+		return FRevisionControlStyleManager::Get().GetBrush("RevisionControl.Icon.ConnectedBadge");
+	}
+	else
+	{
+		return nullptr;
+	}
+}
+
+#undef LOCTEXT_NAMESPACE
\ No newline at end of file
diff --git a/Engine/Source/Editor/SceneOutliner/Private/SceneOutlinerTreeItemSCC.cpp b/Engine/Source/Editor/SceneOutliner/Private/SceneOutlinerTreeItemSCC.cpp
--- a/Engine/Source/Editor/SceneOutliner/Private/SceneOutlinerTreeItemSCC.cpp
+++ b/Engine/Source/Editor/SceneOutliner/Private/SceneOutlinerTreeItemSCC.cpp
@@ -7,6 +7,7 @@
 #include "ActorFolderTreeItem.h"
 #include "ISourceControlProvider.h"
 #include "ISourceControlModule.h"
+#include "PackageTools.h"
 #include "ProfilingDebugging/CpuProfilerTrace.h"
 #include "SceneOutlinerHelpers.h"
 #include "WorldPartition/WorldPartitionActorDesc.h"
@@ -33,9 +34,8 @@ void FSceneOutlinerTreeItemSCC::Initialize()
 {
 	if (TreeItemPtr.IsValid())
 	{
-		ExternalPackageName = SceneOutliner::FSceneOutlinerHelpers::GetExternalPackageName(*TreeItemPtr.Get());
+		ExternalPackageName = TreeItemPtr->GetPackageName();
 		ExternalPackageFileName = !ExternalPackageName.IsEmpty() ? USourceControlHelpers::PackageFilename(ExternalPackageName) : FString();
-		ExternalPackage = SceneOutliner::FSceneOutlinerHelpers::GetExternalPackage(*TreeItemPtr.Get());
 		
 		if (FActorTreeItem* ActorItem = TreeItemPtr->CastTo<FActorTreeItem>())
 		{
@@ -47,14 +47,12 @@ void FSceneOutlinerTreeItemSCC::Initialize()
 					{
 						ExternalPackageName = InActor->GetExternalPackage()->GetName();
 						ExternalPackageFileName = USourceControlHelpers::PackageFilename(ExternalPackageName);
-						ExternalPackage = InActor->GetExternalPackage();
 						ConnectSourceControl();
 					}
 					else
 					{
 						ExternalPackageName = FString();
 						ExternalPackageFileName = FString();
-						ExternalPackage = nullptr;
 						DisconnectSourceControl();
 					}
 				});
@@ -91,6 +89,16 @@ FSourceControlStatePtr FSceneOutlinerTreeItemSCC::RefreshSourceControlState()
 	return ISourceControlModule::Get().GetProvider().GetState(ExternalPackageFileName, EStateCacheUsage::ForceUpdate);
 }
 
+UPackage* FSceneOutlinerTreeItemSCC::GetPackage() const
+{
+	return FindPackage(nullptr, *ExternalPackageName);
+}
+
+UPackage* FSceneOutlinerTreeItemSCC::LoadPackage() const
+{
+	return UPackageTools::LoadPackage(ExternalPackageName);
+}
+
 void FSceneOutlinerTreeItemSCC::ConnectSourceControl()
 {
 	check(!ExternalPackageFileName.IsEmpty());
diff --git a/Engine/Source/Editor/SceneOutliner/Private/SceneOutlinerUnsavedColumn.cpp b/Engine/Source/Editor/SceneOutliner/Private/SceneOutlinerUnsavedColumn.cpp
--- a/Engine/Source/Editor/SceneOutliner/Private/SceneOutlinerUnsavedColumn.cpp
+++ b/Engine/Source/Editor/SceneOutliner/Private/SceneOutlinerUnsavedColumn.cpp
@@ -118,7 +118,7 @@ SHeaderRow::FColumn::FArguments FSceneOutlinerActorUnsavedColumn::ConstructHeade
 
 const TSharedRef<SWidget> FSceneOutlinerActorUnsavedColumn::ConstructRowWidget(FSceneOutlinerTreeItemRef TreeItem, const STableRow<FSceneOutlinerTreeItemPtr>& Row)
 {
-	const FString ExternalPackageName = SceneOutliner::FSceneOutlinerHelpers::GetExternalPackageName(TreeItem.Get());
+	const FString ExternalPackageName = TreeItem->GetPackageName();
 	const FString ExternalPackageFileName = !ExternalPackageName.IsEmpty() ? USourceControlHelpers::PackageFilename(ExternalPackageName) : FString();
 
 	if (ExternalPackageFileName.IsEmpty())
@@ -147,7 +147,7 @@ void FSceneOutlinerActorUnsavedColumn::SortItems(TArray<FSceneOutlinerTreeItemPt
 		/** Sort by unsaved first */
 		.Primary([this](const ISceneOutlinerTreeItem& Item)
 		{
-			FString LongPackageName = SceneOutliner::FSceneOutlinerHelpers::GetExternalPackageName(Item);
+			FString LongPackageName = Item.GetPackageName();
 			if (!LongPackageName.IsEmpty())
 			{
 				if (const TSharedRef<SUnsavedActorWidget>* FoundWidget = UnsavedActorWidgets.Find(USourceControlHelpers::PackageFilename(LongPackageName)))
diff --git a/Engine/Source/Editor/SceneOutliner/Public/ActorDescTreeItem.h b/Engine/Source/Editor/SceneOutliner/Public/ActorDescTreeItem.h
--- a/Engine/Source/Editor/SceneOutliner/Public/ActorDescTreeItem.h
+++ b/Engine/Source/Editor/SceneOutliner/Public/ActorDescTreeItem.h
@@ -68,6 +68,7 @@ public:
 	virtual bool ShouldShowVisibilityState() const override { return false; }
 	virtual bool HasPinnedStateInfo() const override { return true; }
 	virtual bool GetPinnedState() const override;
+	virtual FString GetPackageName() const override;
 	/* End ISceneOutlinerTreeItem Implementation */
 	
 	/* Begin IActorBaseTreeItem Implementation */
diff --git a/Engine/Source/Editor/SceneOutliner/Public/ActorFolderTreeItem.h b/Engine/Source/Editor/SceneOutliner/Public/ActorFolderTreeItem.h
--- a/Engine/Source/Editor/SceneOutliner/Public/ActorFolderTreeItem.h
+++ b/Engine/Source/Editor/SceneOutliner/Public/ActorFolderTreeItem.h
@@ -26,6 +26,7 @@ public:
 	virtual bool ShouldShowVisibilityState() const override;
 	virtual bool ShouldRemoveOnceLastChildRemoved() const override;
 	virtual FFolder GetFolder() const override;
+	virtual FString GetPackageName() const override;
 	/* End FFolderTreeItem Implementation */
 		
 	/* Begin FFolderTreeItem Implementation */
diff --git a/Engine/Source/Editor/SceneOutliner/Public/ActorTreeItem.h b/Engine/Source/Editor/SceneOutliner/Public/ActorTreeItem.h
--- a/Engine/Source/Editor/SceneOutliner/Public/ActorTreeItem.h
+++ b/Engine/Source/Editor/SceneOutliner/Public/ActorTreeItem.h
@@ -52,6 +52,7 @@ public:
 	virtual bool GetPinnedState() const override;
 	virtual void OnLabelChanged() override;
 	virtual void GenerateContextMenu(UToolMenu* Menu, SSceneOutliner& Outliner) override;
+	virtual FString GetPackageName() const override;
 	/* End ISceneOutlinerTreeItem Implementation */
 
 	/* Begin IActorBaseTreeItem Implementation */
diff --git a/Engine/Source/Editor/SceneOutliner/Public/ISceneOutlinerTreeItem.h b/Engine/Source/Editor/SceneOutliner/Public/ISceneOutlinerTreeItem.h
--- a/Engine/Source/Editor/SceneOutliner/Public/ISceneOutlinerTreeItem.h
+++ b/Engine/Source/Editor/SceneOutliner/Public/ISceneOutlinerTreeItem.h
@@ -181,4 +181,7 @@ public:
 
 	/** Whether the item should be removed when the last child has been removed. */
 	virtual bool ShouldRemoveOnceLastChildRemoved() const { check(GetChildren().IsEmpty()); return Flags.bIsFilteredOut; }
+
+	/** Returns (optionally) the name of the UPackage this item represents/is-part-of */
+	virtual FString GetPackageName() const { return FString(); }
 };
diff --git a/Engine/Source/Editor/SceneOutliner/Public/SceneOutlinerHelpers.h b/Engine/Source/Editor/SceneOutliner/Public/SceneOutlinerHelpers.h
--- a/Engine/Source/Editor/SceneOutliner/Public/SceneOutlinerHelpers.h
+++ b/Engine/Source/Editor/SceneOutliner/Public/SceneOutlinerHelpers.h
@@ -10,7 +10,9 @@ namespace SceneOutliner
 	class SCENEOUTLINER_API FSceneOutlinerHelpers
 	{
 	public:
+		UE_DEPRECATED(5.5, "FSceneOutlinerHelpers::GetExternalPackageName has been deprecated implement/use ISceneOutlinerTreeItem::GetPackageName instead")
 		static FString GetExternalPackageName(const ISceneOutlinerTreeItem& TreeItem);
+		UE_DEPRECATED(5.5, "FSceneOutlinerHelpers::GetExternalPackage has been deprecated implement/use ISceneOutlinerTreeItem::GetPackageName instead")
 		static UPackage* GetExternalPackage(const ISceneOutlinerTreeItem& TreeItem);
 		static TSharedPtr<SWidget> GetClassHyperlink(UObject* InObject);
 	};
diff --git a/Engine/Source/Editor/SceneOutliner/Public/SceneOutlinerSourceControlColumn.h b/Engine/Source/Editor/SceneOutliner/Public/SceneOutlinerSourceControlColumn.h
new file mode 100644
--- /dev/null
+++ b/Engine/Source/Editor/SceneOutliner/Public/SceneOutlinerSourceControlColumn.h
@@ -0,0 +1,42 @@
+// Copyright Epic Games, Inc. All Rights Reserved.
+
+#pragma once
+
+#include "CoreMinimal.h"
+#include "Widgets/SWidget.h"
+#include "Widgets/Views/SHeaderRow.h"
+#include "ISceneOutliner.h"
+#include "SceneOutlinerPublicTypes.h"
+#include "ISceneOutlinerColumn.h"
+
+template<typename ItemType> class STableRow;
+
+/** A column for the SceneOutliner that displays the SCC Information */
+class SCENEOUTLINER_API FSceneOutlinerSourceControlColumn : public ISceneOutlinerColumn
+{
+
+public:
+	FSceneOutlinerSourceControlColumn(ISceneOutliner& SceneOutliner) : WeakSceneOutliner(StaticCastSharedRef<ISceneOutliner>(SceneOutliner.AsShared())) {}
+
+	virtual ~FSceneOutlinerSourceControlColumn() {}
+
+	static FName GetID() { return FSceneOutlinerBuiltInColumnTypes::SourceControl(); }
+
+	//////////////////////////////////////////////////////////////////////////
+	// Begin ISceneOutlinerColumn Implementation
+
+	virtual FName GetColumnID() override;
+
+	virtual SHeaderRow::FColumn::FArguments ConstructHeaderRowColumn() override;
+
+	virtual const TSharedRef<SWidget> ConstructRowWidget(FSceneOutlinerTreeItemRef TreeItem, const STableRow<FSceneOutlinerTreeItemPtr>& Row) override;
+
+	virtual bool SupportsSorting() const override { return false; }
+
+	// End ISceneOutlinerColumn Implementation
+	//////////////////////////////////////////////////////////////////////////
+private:
+	const FSlateBrush* GetHeaderIconBadge() const;
+
+	TWeakPtr<ISceneOutliner> WeakSceneOutliner;
+};
diff --git a/Engine/Source/Editor/SceneOutliner/Public/SceneOutlinerTreeItemSCC.h b/Engine/Source/Editor/SceneOutliner/Public/SceneOutlinerTreeItemSCC.h
--- a/Engine/Source/Editor/SceneOutliner/Public/SceneOutlinerTreeItemSCC.h
+++ b/Engine/Source/Editor/SceneOutliner/Public/SceneOutlinerTreeItemSCC.h
@@ -17,42 +17,29 @@ class FSceneOutlinerTreeItemSCC : public TSharedFromThis<FSceneOutlinerTreeItemS
 {
 public:
 	FSceneOutlinerTreeItemSCC(FSceneOutlinerTreeItemPtr InTreeItemPtr);
-
 	~FSceneOutlinerTreeItemSCC();
 
 	// Actually attempt to connect to Revision Control and get the item's status
 	void Initialize();
 
 	FSourceControlStatePtr GetSourceControlState();
-
 	FSourceControlStatePtr RefreshSourceControlState();
+	TWeakPtr<FUncontrolledChangelistState> GetUncontrolledChangelistState() { return UncontrolledChangelistState; }
 
-	bool IsExternalPackage() { return !ExternalPackageName.IsEmpty(); }
-
-	FString GetPackageName() { return ExternalPackageName; }
-
-	FString GetPackageFileName() { return ExternalPackageFileName; }
-
-	UPackage* GetPackage() { return ExternalPackage; }
+	bool HasValidPackage() const { return !ExternalPackageName.IsEmpty(); }
+	FString GetPackageName() const { return ExternalPackageName; }
+	FString GetPackageFileName() const { return ExternalPackageFileName; }
+	UPackage* GetPackage() const;
+	UPackage* LoadPackage() const;
 
 	FSourceControlStateChangedDelegate OnSourceControlStateChanged;
-	
 	FUncontrolledStateChangedDelegate OnUncontrolledChangelistsStateChanged;
-	
-	TWeakPtr<FUncontrolledChangelistState> GetUncontrolledChangelistState() { return UncontrolledChangelistState; }
-
 private:
-
 	void ConnectSourceControl();
-
 	void DisconnectSourceControl();
-
 	void HandleSourceControlStateChanged(EStateCacheUsage::Type CacheUsage);
-
 	void HandleSourceControlProviderChanged(ISourceControlProvider& OldProvider, ISourceControlProvider& NewProvider);
-
 	void BroadcastNewState(FSourceControlStatePtr SourceControlState);
-
 	void HandleUncontrolledChangelistsStateChanged();
 
 	/** The tree item we relate to */
@@ -61,7 +48,6 @@ private:
 	/** Cache the items external package name and filename */
 	FString ExternalPackageName;
 	FString ExternalPackageFileName;
-	UPackage* ExternalPackage = nullptr;
 
 	/** Source control state changed delegate handle */
 	FDelegateHandle SourceControlStateChangedDelegateHandle;
