WorkspaceEditor: unify dirty state icon usages

diff --git a/Engine/Plugins/Experimental/Workspace/Source/WorkspaceEditor/Private/Outliner/WorkspaceOutlinerColumns.cpp b/Engine/Plugins/Experimental/Workspace/Source/WorkspaceEditor/Private/Outliner/WorkspaceOutlinerColumns.cpp
--- a/Engine/Plugins/Experimental/Workspace/Source/WorkspaceEditor/Private/Outliner/WorkspaceOutlinerColumns.cpp
+++ b/Engine/Plugins/Experimental/Workspace/Source/WorkspaceEditor/Private/Outliner/WorkspaceOutlinerColumns.cpp
@@ -39,7 +39,7 @@ namespace UE::Workspace
 			[
 				SNew(SImage)
 				.ColorAndOpacity(FSlateColor::UseForeground())
-				.Image(FAppStyle::GetBrush("ContentBrowser.ContentDirty"))
+				.Image(FAppStyle::GetBrush("Icons.DirtyBadge"))
 				.ToolTipText(LOCTEXT("FileStatusTooltip", "File status of this entry"))
 			]
 		];	
@@ -91,7 +91,7 @@ namespace UE::Workspace
 					{
 						if (Package->IsDirty())
 						{
-							return FAppStyle::GetBrush("ContentBrowser.ContentDirty");
+							return FAppStyle::GetBrush("Icons.DirtyBadge");
 						}
 					}
 					
