Workspace: saving/dirty state UX improvements:

Modified/Dirty state:
- Added custom behavior for tab-suffix (showing dirty Workspace asset)
- Breadcrumbs now show dirty asterix in-line

Saving:
- Add "Save selected" to outliner context menu
- Breadcrumb changes
    * Added CanSave/Save delegates
    * Use new delegates to populate a mini toolbar, with single save button, inside of the TabWrapper
    * Use new delegates to generate Breadcrumb widget content (showing dirty state icon in-line)
    * Added delegate bindings for all current asset types
- Added Save Asset Entries toolbar button to Workspace, this saves all modified assets (and external packages) contained by the workspace
- Added Save Asset button to tabwrapper (document navigation bar), this saves the currently open asset (and any of its external packages)
- Saving transient workspace now has a suppressible dialog prompt

Misc:
- Make outliner multi-selection
- Improve label/tooltips (single/plural forms)
- Make Open Animation Graph available for multi-select
- Hid find-in-content-browser toolbar in Workspace editor
- Made FWorkspaceEditor::OpenObjects open new documents if opening multiple objects at the same time, otherwise they would consecutively get opened in the same document tab

diff --git a/Engine/Plugins/Experimental/AnimNext/Source/AnimNextEditor/Private/AnimNextEditorModule.cpp b/Engine/Plugins/Experimental/AnimNext/Source/AnimNextEditor/Private/AnimNextEditorModule.cpp
--- a/Engine/Plugins/Experimental/AnimNext/Source/AnimNextEditor/Private/AnimNextEditorModule.cpp
+++ b/Engine/Plugins/Experimental/AnimNext/Source/AnimNextEditor/Private/AnimNextEditorModule.cpp
@@ -34,6 +34,7 @@
 #include "AnimNextRigVMAssetEntry.h"
 #include "Editor/RigVMGraphDetailCustomization.h"
 #include "EditorUtils.h"
+#include "FileHelpers.h"
 #include "IUniversalObjectLocatorEditorModule.h"
 #include "Param/AnimNextComponentLocatorEditor.h"
 #include "Param/AnimNextParam.h"
@@ -339,6 +340,29 @@ void FModule::RegisterWorkspaceDocumentTypes(Workspace::IWorkspaceEditorModule&
 		{
 			const TSharedPtr<Workspace::FWorkspaceBreadcrumb>& GraphCrumb = OutBreadcrumbs.Add_GetRef(MakeShared<Workspace::FWorkspaceBreadcrumb>());
 			GraphCrumb->OnGetLabel = Workspace::FWorkspaceBreadcrumb::FOnGetBreadcrumbLabel::CreateLambda([ScheduleName = Schedule->GetFName()]{ return FText::FromName(ScheduleName); });
+
+			TWeakObjectPtr<const UAnimNextSchedule> WeakSchedule = Schedule;
+			GraphCrumb->CanSave = Workspace::FWorkspaceBreadcrumb::FCanSaveBreadcrumb::CreateLambda(
+				[WeakSchedule]
+				{
+					if (const UAnimNextSchedule* Schedule = WeakSchedule.Get())
+					{
+						return Schedule->GetPackage()->IsDirty();
+					}
+
+					return false;
+				}
+			);
+
+			GraphCrumb->OnSave = Workspace::FWorkspaceBreadcrumb::FOnSaveBreadcrumb::CreateLambda(
+				[WeakSchedule]
+				{
+					if (const UAnimNextSchedule* Schedule = WeakSchedule.Get())
+					{
+						FEditorFileUtils::PromptForCheckoutAndSave({Schedule->GetPackage()}, false, /*bPromptToSave=*/ false);
+					}
+				}
+			);
 		}
 	});
 
@@ -422,8 +446,29 @@ void FModule::RegisterWorkspaceDocumentTypes(Workspace::IWorkspaceEditorModule&
 	{
 		if (const UAnimNextModule* Module = Cast<UAnimNextModule>(InContext.Object))
 		{
+			TWeakObjectPtr<const UAnimNextModule> WeakModule = Module;
 			const TSharedPtr<Workspace::FWorkspaceBreadcrumb>& GraphCrumb = OutBreadcrumbs.Add_GetRef(MakeShared<Workspace::FWorkspaceBreadcrumb>());
 			GraphCrumb->OnGetLabel = Workspace::FWorkspaceBreadcrumb::FOnGetBreadcrumbLabel::CreateLambda([ModuleName = Module->GetFName()]{ return FText::FromName(ModuleName); });
+			GraphCrumb->CanSave = Workspace::FWorkspaceBreadcrumb::FCanSaveBreadcrumb::CreateLambda(
+				[WeakModule]
+				{
+					if (const UAnimNextModule* Module = WeakModule.Get())
+					{
+						return Module->GetPackage()->IsDirty();						
+					}
+
+					return false;
+				}
+			);
+			GraphCrumb->OnSave = Workspace::FWorkspaceBreadcrumb::FOnSaveBreadcrumb::CreateLambda(
+				[WeakModule]
+				{
+					if (const UAnimNextModule* Module = WeakModule.Get())
+					{
+						FEditorFileUtils::PromptForCheckoutAndSave({Module->GetPackage()}, false, /*bPromptToSave=*/ false);
+					}
+				}
+			);
 		}
 	});
 
@@ -788,11 +833,34 @@ void FModule::RegisterWorkspaceDocumentTypes(Workspace::IWorkspaceEditorModule&
 	{
 		if (const URigVMEdGraph* RigVMEdGraph = Cast<URigVMEdGraph>(InContext.Object))
 		{
-			const TSharedPtr<Workspace::FWorkspaceBreadcrumb>& GraphCrumb = OutBreadcrumbs.Add_GetRef(MakeShared<Workspace::FWorkspaceBreadcrumb>());
-			GraphCrumb->OnGetLabel = Workspace::FWorkspaceBreadcrumb::FOnGetBreadcrumbLabel::CreateLambda([GraphName = RigVMEdGraph->GetFName()]{ return FText::FromName(GraphName); });
-
 			if (const UAnimNextRigVMAssetEntry* OuterAssetEntry = CastChecked<UAnimNextRigVMAssetEntry>(RigVMEdGraph->GetOuter()))
 			{
+				const IAnimNextRigVMGraphInterface* GraphInterface = Cast<IAnimNextRigVMGraphInterface>(OuterAssetEntry);
+				const TSharedPtr<Workspace::FWorkspaceBreadcrumb>& GraphCrumb = OutBreadcrumbs.Add_GetRef(MakeShared<Workspace::FWorkspaceBreadcrumb>());
+
+				TWeakObjectPtr<const URigVMEdGraph> WeakGraph = RigVMEdGraph;
+				GraphCrumb->OnGetLabel = Workspace::FWorkspaceBreadcrumb::FOnGetBreadcrumbLabel::CreateLambda([GraphName = OuterAssetEntry->GetEntryName()]{ return FText::FromName(GraphName); });
+				GraphCrumb->CanSave = Workspace::FWorkspaceBreadcrumb::FCanSaveBreadcrumb::CreateLambda(
+					[WeakGraph]
+					{
+						if (const URigVMEdGraph* Graph = WeakGraph.Get())
+						{
+							return Graph->GetPackage()->IsDirty();
+						}
+						return false;
+					}
+				);
+
+				GraphCrumb->OnSave = Workspace::FWorkspaceBreadcrumb::FOnSaveBreadcrumb::CreateLambda(
+				[WeakGraph]
+					{
+						if (const URigVMEdGraph* Graph = WeakGraph.Get())
+						{
+							FEditorFileUtils::PromptForCheckoutAndSave({Graph->GetPackage()}, false, /*bPromptToSave=*/ false);
+						}
+					}
+				);
+	
 				if (UAnimNextModule_EditorData* OuterEditorData = CastChecked<UAnimNextModule_EditorData>(OuterAssetEntry->GetOuter()))
 				{
 					if(UAnimNextModule* OuterGraph = UncookedOnly::FUtils::GetGraph(OuterEditorData))
@@ -810,6 +878,28 @@ void FModule::RegisterWorkspaceDocumentTypes(Workspace::IWorkspaceEditorModule&
 								}
 							}
 						);
+						OuterGraphCrumb->CanSave = Workspace::FWorkspaceBreadcrumb::FCanSaveBreadcrumb::CreateLambda(
+							[WeakOuterGraph]
+							{
+								if (UAnimNextModule* Module = WeakOuterGraph.Get())
+								{
+									return Module->GetPackage()->IsDirty();
+								}
+
+								return false;
+							}
+						);
+
+						OuterGraphCrumb->OnSave = Workspace::FWorkspaceBreadcrumb::FOnSaveBreadcrumb::CreateLambda(
+						[WeakOuterGraph]
+							{
+								if (UAnimNextModule* Module = WeakOuterGraph.Get())
+								{
+									FEditorFileUtils::PromptForCheckoutAndSave({Module->GetPackage()}, false, /*bPromptToSave=*/ false);
+								}
+							}
+						);
+						
 					}
 				}
 			}
diff --git a/Engine/Plugins/Experimental/AnimNext/Source/AnimNextEditor/Private/Graph/AnimNextGraphItemDetails.cpp b/Engine/Plugins/Experimental/AnimNext/Source/AnimNextEditor/Private/Graph/AnimNextGraphItemDetails.cpp
--- a/Engine/Plugins/Experimental/AnimNext/Source/AnimNextEditor/Private/Graph/AnimNextGraphItemDetails.cpp
+++ b/Engine/Plugins/Experimental/AnimNext/Source/AnimNextEditor/Private/Graph/AnimNextGraphItemDetails.cpp
@@ -80,40 +80,57 @@ void FAnimNextGraphItemDetails::RegisterToolMenuExtensions()
 				FToolMenuSection& Section = InMenu->AddSection("WorkspaceOutliner.ItemContextMenu.RootAsset", FText::FromString(TEXT("Animation Next")));
 				if(const TSharedPtr<UE::Workspace::IWorkspaceEditor> WorkspaceEditor = StaticCastSharedPtr<UE::Workspace::IWorkspaceEditor>(AssetEditorContext->Toolkit.Pin()))
 				{
-					if (WorkspaceItemContext->SelectedExports.Num() == 1)
+
+					TArray<FWorkspaceOutlinerItemExport> GraphExports;
+					Algo::TransformIf(WorkspaceItemContext->SelectedExports, GraphExports, [](const FWorkspaceOutlinerItemExport& Export)
+					{
+						return Export.Data.IsValid() && Export.Data.GetScriptStruct() == FAnimNextGraphOutlinerData::StaticStruct();
+					},
+					[](const FWorkspaceOutlinerItemExport& Export)
+					{
+						return Export;
+					});
+
+					if (GraphExports.Num() == WorkspaceItemContext->SelectedExports.Num() && GraphExports.Num() > 0)
 					{
-						TInstancedStruct<FWorkspaceOutlinerItemData>& Data = WorkspaceItemContext->SelectedExports[0].Data;
-						if (Data.IsValid() && Data.GetScriptStruct() == FAnimNextGraphOutlinerData::StaticStruct())
-						{							
-							Section.AddMenuEntry(
-							TEXT("Open"),
-							LOCTEXT("GoHereActionLabel", "Open Graph"),
-							LOCTEXT("GoHereActionTooltip", "Open selected Animation Graph"),
-							FSlateIcon(FAppStyle::GetAppStyleSetName(), "GraphEditor.EventGraph_16x"),
-							FUIAction(
-								FExecuteAction::CreateWeakLambda(WorkspaceItemContext, [&Data, WorkspaceEditor]()
-								{										
-									const FAnimNextGraphOutlinerData& GraphData = Data.Get<FAnimNextGraphOutlinerData>();
-									if (GraphData.GraphInterface)
+						TWeakPtr<UE::Workspace::IWorkspaceEditor> WeakWorkspaceEditor = WorkspaceEditor;
+						Section.AddMenuEntry(
+						TEXT("OpenGraphMenuEntry"),
+						FText::FormatOrdered(LOCTEXT("OpenGraphMenuEntryLabel", "Open Animation {0}|plural(one=Graph,other=Graphs)"), GraphExports.Num()),
+						FText::FormatOrdered(LOCTEXT("OpenGraphMenuEntryTooltip", "Open the selected Animation {0}|plural(one=Graph,other=Graphs)"), GraphExports.Num()),
+						FSlateIcon(FAppStyle::GetAppStyleSetName(), "GraphEditor.EventGraph_16x"),
+						FUIAction(
+							FExecuteAction::CreateWeakLambda(WorkspaceItemContext, [GraphExports, WeakWorkspaceEditor]()
+							{
+								if (TSharedPtr<UE::Workspace::IWorkspaceEditor> SharedWorkspaceEditor = WeakWorkspaceEditor.Pin())
+								{
+									TArray<UObject*> ObjectsToOpen;
+									for (const FWorkspaceOutlinerItemExport& Export : GraphExports)
 									{
-										if (URigVMGraph* RigVMGraph = GraphData.GraphInterface->GetRigVMGraph())
+										const FAnimNextGraphOutlinerData& GraphData = Export.Data.Get<FAnimNextGraphOutlinerData>();
+										if (GraphData.GraphInterface)
 										{
-											if(const IRigVMClientHost* RigVMClientHost = RigVMGraph->GetImplementingOuter<IRigVMClientHost>())
+											if (URigVMGraph* RigVMGraph = GraphData.GraphInterface->GetRigVMGraph())
 											{
-												if(UObject* EditorObject = RigVMClientHost->GetEditorObjectForRigVMGraph(RigVMGraph))
+												if(const IRigVMClientHost* RigVMClientHost = RigVMGraph->GetImplementingOuter<IRigVMClientHost>())
 												{
-													WorkspaceEditor->OpenObjects({EditorObject});
+													if(UObject* EditorObject = RigVMClientHost->GetEditorObjectForRigVMGraph(RigVMGraph))
+													{
+														ObjectsToOpen.Add(EditorObject);
+													}
 												}
 											}
 										}
 									}
-								})
-							));
-						}
+
+									SharedWorkspaceEditor->OpenObjects({ObjectsToOpen});
+								}
+							})
+						));
 					}
 				}
 			}
-		}));
+		}), FToolMenuInsert(NAME_None, EToolMenuInsertType::First));
 	}
 }
 
diff --git a/Engine/Plugins/Experimental/Workspace/Source/WorkspaceEditor/Private/Outliner/WorkspaceOutlinerMode.cpp b/Engine/Plugins/Experimental/Workspace/Source/WorkspaceEditor/Private/Outliner/WorkspaceOutlinerMode.cpp
--- a/Engine/Plugins/Experimental/Workspace/Source/WorkspaceEditor/Private/Outliner/WorkspaceOutlinerMode.cpp
+++ b/Engine/Plugins/Experimental/Workspace/Source/WorkspaceEditor/Private/Outliner/WorkspaceOutlinerMode.cpp
@@ -2,6 +2,7 @@
 
 #include "WorkspaceOutlinerMode.h"
 
+#include "FileHelpers.h"
 #include "ISourceControlModule.h"
 #include "WorkspaceItemMenuContext.h"
 #include "WorkspaceOutlinerHierarchy.h"
@@ -70,12 +71,12 @@ TSharedPtr<SWidget> FWorkspaceOutlinerMode::CreateContextMenu()
 						return Export.ParentIdentifier == NAME_None;
 					});
 
+					FToolMenuSection& AssetsSection = InMenu->AddSection("Assets", LOCTEXT("AssetSectionLabel", "Assets"));
 					if (bSelectionContainsTopLevelAsset)
 					{
-						FToolMenuSection& AssetsSection = InMenu->AddSection("Assets", LOCTEXT("AssetSectionLabel", "Assets"));
 						AssetsSection.AddMenuEntry(TEXT("OpenAsset"),
-							LOCTEXT("OpenAssetLabel", "Open Asset(s)"),
-							LOCTEXT("OpenAssetTooltip", "Opens the selected Asset(s)"),
+							FText::FormatOrdered(LOCTEXT("OpenAssetLabel", "Open {0}|plural(one=Asset,other=Assets)"), MenuContext->SelectedExports.Num()),
+							FText::FormatOrdered(LOCTEXT("OpenAssetTooltip", "Opens the selected {0}|plural(one=Asset,other=Assets)"), MenuContext->SelectedExports.Num()),
 							FSlateIcon(FAppStyle::Get().GetStyleSetName(), "SystemWideCommands.SummonOpenAssetDialog"),
 							FUIAction(FExecuteAction::CreateLambda([WeakEditor=EditorContext->Toolkit, SelectedExports=MenuContext->SelectedExports]()
 							{
@@ -99,8 +100,8 @@ TSharedPtr<SWidget> FWorkspaceOutlinerMode::CreateContextMenu()
 						);
 
 						AssetsSection.AddMenuEntry(TEXT("RemoveAsset"),
-							LOCTEXT("RemoveAssetLabel", "Remove Asset(s)"),
-							LOCTEXT("RemoveAssetTooltip", "Removes the selected Asset(s) from the Workspace"),
+							FText::FormatOrdered(LOCTEXT("RemoveAssetLabel", "Remove {0}|plural(one=Asset,other=Assets)"), MenuContext->SelectedExports.Num()),
+							FText::FormatOrdered(LOCTEXT("RemoveAssetTooltip", "Removes the selected {0}|plural(one=Asset,other=Assets) from the Workspace"), MenuContext->SelectedExports.Num()),
 							FSlateIcon(FAppStyle::Get().GetStyleSetName(), "Icons.Delete"),
 							FUIAction(FExecuteAction::CreateLambda([WeakEditor=EditorContext->Toolkit, SelectedExports=MenuContext->SelectedExports, EditingObjects=EditorContext->GetEditingObjects()]()
 							{
@@ -129,6 +130,79 @@ TSharedPtr<SWidget> FWorkspaceOutlinerMode::CreateContextMenu()
 							}))
 						);
 					}
+
+					auto IsPackageDirty = [](const UPackage* Package)-> bool
+					{										
+						return (Package && (Package->IsDirty() || Package->GetExternalPackages().ContainsByPredicate([](const UPackage* Package) { return Package->IsDirty(); })));
+					};
+
+					AssetsSection.AddMenuEntry("SaveSelectedAssets",
+					FText::FormatOrdered(LOCTEXT("SaveSelectedAssets", "Save {0}|plural(one=Asset,other=Assets)"), MenuContext->SelectedExports.Num()),
+					FText::FormatOrdered(LOCTEXT("SaveSelectedAssets_ToolTip", "Save the selected {0}|plural(one=Asset,other=Assets)"), MenuContext->SelectedExports.Num()),
+					FSlateIcon(FAppStyle::Get().GetStyleSetName(), "AssetEditor.SaveAsset"),
+					FUIAction(
+						FExecuteAction::CreateLambda([IsPackageDirty, WeakEditor=EditorContext->Toolkit, SelectedExports=MenuContext->SelectedExports, EditingObjects=EditorContext->GetEditingObjects()]()
+						{
+							if (const TSharedPtr<UE::Workspace::IWorkspaceEditor> SharedWorkspaceEditor = StaticCastSharedPtr<UE::Workspace::IWorkspaceEditor>(WeakEditor.Pin()))
+							{
+								TArray<UPackage*> SavablePackages;
+								for (const FWorkspaceOutlinerItemExport& ItemExport : SelectedExports)
+								{
+									if (const TSharedPtr<IWorkspaceOutlinerItemDetails> SharedFactory = FWorkspaceEditorModule::GetOutlinerItemDetails(MakeOutlinerDetailsId(ItemExport)))
+									{
+										UPackage* Package = SharedFactory->GetPackage(ItemExport);
+										if (IsPackageDirty(Package))
+										{
+											SavablePackages.AddUnique(Package);								
+										}										
+									}
+									else if (ItemExport.ParentIdentifier == NAME_None)
+									{
+										if (UPackage* Package = FindPackage(nullptr, *ItemExport.AssetPath.GetLongPackageName()))
+										{
+											if (IsPackageDirty(Package))
+											{
+												SavablePackages.AddUnique(Package);
+											}
+										}
+									}
+								}
+								
+								FEditorFileUtils::PromptForCheckoutAndSave(SavablePackages, false, /*bPromptToSave=*/ false);
+							}
+						}),
+						FCanExecuteAction::CreateLambda([IsPackageDirty, WeakEditor=EditorContext->Toolkit, SelectedExports=MenuContext->SelectedExports, EditingObjects=EditorContext->GetEditingObjects()]()
+						{
+							if (const TSharedPtr<UE::Workspace::IWorkspaceEditor> SharedWorkspaceEditor = StaticCastSharedPtr<UE::Workspace::IWorkspaceEditor>(WeakEditor.Pin()))
+							{
+								TArray<UPackage*> SavablePackages;
+								for (const FWorkspaceOutlinerItemExport& ItemExport : SelectedExports)
+								{
+									if (const TSharedPtr<IWorkspaceOutlinerItemDetails> SharedFactory = FWorkspaceEditorModule::GetOutlinerItemDetails(MakeOutlinerDetailsId(ItemExport)))
+									{									
+										const UPackage* Package = SharedFactory->GetPackage(ItemExport);
+										if (IsPackageDirty(Package))
+										{
+											return true;
+										}
+									}
+									else if (ItemExport.ParentIdentifier == NAME_None)
+									{
+										if (const UPackage* Package = FindPackage(nullptr, *ItemExport.AssetPath.GetLongPackageName()))
+										{
+											if (IsPackageDirty(Package))
+											{
+												return true;
+											}
+										}
+									}
+								}
+							}
+
+							return false;
+						})
+					));
+
 				}
 
 				if (TSharedPtr<SSceneOutliner> SharedOutliner = WeakOutliner.Pin())
diff --git a/Engine/Plugins/Experimental/Workspace/Source/WorkspaceEditor/Private/Outliner/WorkspaceOutlinerMode.h b/Engine/Plugins/Experimental/Workspace/Source/WorkspaceEditor/Private/Outliner/WorkspaceOutlinerMode.h
--- a/Engine/Plugins/Experimental/Workspace/Source/WorkspaceEditor/Private/Outliner/WorkspaceOutlinerMode.h
+++ b/Engine/Plugins/Experimental/Workspace/Source/WorkspaceEditor/Private/Outliner/WorkspaceOutlinerMode.h
@@ -25,6 +25,7 @@ namespace UE::Workspace
 		void HandleItemSelection(const FSceneOutlinerItemSelection& Selection);
 		virtual void OnItemSelectionChanged(FSceneOutlinerTreeItemPtr Item, ESelectInfo::Type SelectionType, const FSceneOutlinerItemSelection& Selection) override;
 		virtual bool CanCustomizeToolbar() const { return true; }
+		virtual ESelectionMode::Type GetSelectionMode() const { return ESelectionMode::Multi; }
 	protected:
 		virtual TUniquePtr<ISceneOutlinerHierarchy> CreateHierarchy() override;		
 		// End ISceneOutlinerMode overrides
diff --git a/Engine/Plugins/Experimental/Workspace/Source/WorkspaceEditor/Private/SWorkspaceTabWrapper.cpp b/Engine/Plugins/Experimental/Workspace/Source/WorkspaceEditor/Private/SWorkspaceTabWrapper.cpp
--- a/Engine/Plugins/Experimental/Workspace/Source/WorkspaceEditor/Private/SWorkspaceTabWrapper.cpp
+++ b/Engine/Plugins/Experimental/Workspace/Source/WorkspaceEditor/Private/SWorkspaceTabWrapper.cpp
@@ -3,6 +3,9 @@
 #include "SWorkspaceTabWrapper.h"
 #include "Widgets/Layout/SSeparator.h"
 #include "Modules/ModuleManager.h"
+#include "ToolMenus.h"
+
+#define LOCTEXT_NAMESPACE "SWorkspaceTabWrapper"
 
 void SWorkspaceTabWrapper::Construct( const FArguments& InArgs, TSharedPtr<class FTabInfo> InTabInfo, TSharedPtr<UE::Workspace::FWorkspaceEditor> InWorkspaceEditor, UObject* InDocumentID)
 {
@@ -13,6 +16,26 @@ void SWorkspaceTabWrapper::Construct( const FArguments& InArgs, TSharedPtr<class
 	// Set-up shared breadcrumb defaults JDB TODO figure out correct padding to align fake title with breadcrumbs
     const FMargin BreadcrumbTrailPadding = FMargin(4.f, 2.f);
     const FSlateBrush* BreadcrumbButtonImage = FAppStyle::GetBrush("BreadcrumbTrail.Delimiter");
+
+	UToolMenus* ToolMenus = UToolMenus::Get();
+	UToolMenu* ToolBarMenu = nullptr;
+	if(ToolMenus)
+	{
+		const FName ToolbarName = "WorkspaceTabWrapperToolbar";		
+		if (ToolMenus->IsMenuRegistered(ToolbarName))
+		{
+			ToolBarMenu = ToolMenus->ExtendMenu(ToolbarName);
+		}
+		else
+		{
+			ToolBarMenu = ToolMenus->RegisterMenu(ToolbarName, NAME_None, EMultiBoxType::SlimHorizontalToolBar);
+			ToolBarMenu->StyleName = "AssetEditorToolbar";
+		}
+		
+		FToolMenuSection& Section = ToolBarMenu->AddSection("AssetActions");
+		Section.AddMenuEntry("Save", TAttribute<FText>(), LOCTEXT("SaveButtonTooltip", "Save Asset"), FSlateIcon(FAppStyle::Get().GetStyleSetName(), "AssetEditor.SaveAsset"),
+			FUIAction(FExecuteAction::CreateSP(this, &SWorkspaceTabWrapper::ExecuteSave), FCanExecuteAction::CreateSP(this, &SWorkspaceTabWrapper::CanExecuteSave), FIsActionChecked(), FIsActionButtonVisible::CreateSP(this, &SWorkspaceTabWrapper::IsSaveButtonVisible)));
+	}
 	
 	ChildSlot
 	[
@@ -42,10 +65,11 @@ void SWorkspaceTabWrapper::Construct( const FArguments& InArgs, TSharedPtr<class
 					SNew(SHorizontalBox)
 					+SHorizontalBox::Slot()
 					.AutoWidth()
-					//.Padding( 10.0f,5.0f )
+					.Padding( 10.0f,5.0f )
 					.VAlign(VAlign_Center)
 					[
 						SNew(SImage)
+						.DesiredSizeOverride(FVector2D{20.f})
 						.Image( this, &SWorkspaceTabWrapper::GetTabIcon )
 						.ColorAndOpacity(FSlateColor::UseForeground())
 					]
@@ -99,13 +123,56 @@ void SWorkspaceTabWrapper::Construct( const FArguments& InArgs, TSharedPtr<class
 								.TextStyle(FAppStyle::Get(), "GraphBreadcrumbButtonText")
 								.ButtonContentPadding(BreadcrumbTrailPadding)
 								.DelimiterImage(BreadcrumbButtonImage)
-								.PersistentBreadcrumbs(false)
 								.OnCrumbClicked_Lambda([](const TSharedPtr<UE::Workspace::FWorkspaceBreadcrumb> InBreadcrumb){ InBreadcrumb->OnClicked.ExecuteIfBound(); })
+								.GetCrumbButtonContent_Lambda([](const TSharedPtr<UE::Workspace::FWorkspaceBreadcrumb> InBreadcrumb, const FTextBlockStyle* InTextStyle) -> TSharedRef<SWidget>
+								{
+									FText Text = InBreadcrumb->OnGetLabel.Execute().Get();
+									return SNew(SHorizontalBox)
+										+SHorizontalBox::Slot()
+										.AutoWidth()
+										[
+											SNew(STextBlock)
+											.Text(Text)
+											.TextStyle(InTextStyle)
+										]
+										+SHorizontalBox::Slot()
+										.AutoWidth()
+										.Padding(3.0f, 0.f, 0.f, 0.f)
+										[
+											SNew(SBox)
+										   .HAlign(HAlign_Center)
+										   .VAlign(VAlign_Center)
+										   .HeightOverride(20.0f)		
+										   [
+											   SNew(SImage)
+											   .Image_Lambda([InBreadcrumb]() -> const FSlateBrush*
+											   {											
+												   if (InBreadcrumb->CanSave.IsBound() && InBreadcrumb->CanSave.Execute())
+												   {
+													   return FAppStyle::GetBrush("Icons.DirtyBadge");
+												   }
+										   		
+												   return nullptr;
+											   })
+										   ]
+										];
+								})
 							]
 						]
 					]
 				]
 			]
+			+ SHorizontalBox::Slot()
+			.AutoWidth()
+			[
+				SNew(SSeparator)
+				.Orientation(Orient_Vertical)
+			]
+			+SHorizontalBox::Slot()
+			.AutoWidth()
+			[
+				ToolMenus && ToolBarMenu ? ToolMenus->GenerateWidget(ToolBarMenu) : SNullWidget::NullWidget	
+			]
 		]
 		+ SVerticalBox::Slot()
 		.AutoHeight()
@@ -197,5 +264,35 @@ FText  SWorkspaceTabWrapper::GetWorkspaceName() const
 	return FText::GetEmpty();
 }
 
-	
+void SWorkspaceTabWrapper::ExecuteSave() const
+{
+	if(BreadcrumbTrail.IsValid() && BreadcrumbTrail->HasCrumbs())
+	{
+		TSharedPtr<UE::Workspace::FWorkspaceBreadcrumb> LastBreadcrumb = BreadcrumbTrail->PeekCrumb();
+		 LastBreadcrumb->OnSave.ExecuteIfBound();
+	}
+}
+
+bool SWorkspaceTabWrapper::CanExecuteSave() const
+{
+	if(BreadcrumbTrail.IsValid() && BreadcrumbTrail->HasCrumbs())
+	{
+		TSharedPtr<UE::Workspace::FWorkspaceBreadcrumb> LastBreadcrumb = BreadcrumbTrail->PeekCrumb();
+		return LastBreadcrumb->CanSave.IsBound() ? LastBreadcrumb->CanSave.Execute() : false;
+	}
+
+	return false;
+}
+
+bool SWorkspaceTabWrapper::IsSaveButtonVisible() const
+{
+	if(BreadcrumbTrail.IsValid() && BreadcrumbTrail->HasCrumbs())
+	{
+		TSharedPtr<UE::Workspace::FWorkspaceBreadcrumb> LastBreadcrumb = BreadcrumbTrail->PeekCrumb();
+		return LastBreadcrumb->CanSave.IsBound();
+	}
+
+	return false;
+}
 
+#undef LOCTEXT_NAMESPACE // "SWorkspaceTabWrapper"
diff --git a/Engine/Plugins/Experimental/Workspace/Source/WorkspaceEditor/Private/SWorkspaceTabWrapper.h b/Engine/Plugins/Experimental/Workspace/Source/WorkspaceEditor/Private/SWorkspaceTabWrapper.h
--- a/Engine/Plugins/Experimental/Workspace/Source/WorkspaceEditor/Private/SWorkspaceTabWrapper.h
+++ b/Engine/Plugins/Experimental/Workspace/Source/WorkspaceEditor/Private/SWorkspaceTabWrapper.h
@@ -27,6 +27,10 @@ protected:
 	EVisibility IsWorkspaceNameVisible() const;
 	FText GetWorkspaceName() const;
 
+	void ExecuteSave() const;
+	bool CanExecuteSave() const;
+	bool IsSaveButtonVisible() const;
+
 	TSharedPtr<SWidget> Content = nullptr;
 	TWeakPtr<UE::Workspace::FWorkspaceEditor> WeakWorkspaceEditor = nullptr;
 	TWeakObjectPtr<UObject> WeakDocumentObject = nullptr;
diff --git a/Engine/Plugins/Experimental/Workspace/Source/WorkspaceEditor/Private/WorkspaceEditor.cpp b/Engine/Plugins/Experimental/Workspace/Source/WorkspaceEditor/Private/WorkspaceEditor.cpp
--- a/Engine/Plugins/Experimental/Workspace/Source/WorkspaceEditor/Private/WorkspaceEditor.cpp
+++ b/Engine/Plugins/Experimental/Workspace/Source/WorkspaceEditor/Private/WorkspaceEditor.cpp
@@ -16,7 +16,9 @@
 #include "ToolMenuContext.h"
 #include "WorkspaceEditorCommands.h"
 #include "SWorkspaceTabWrapper.h"
+#include "Dialogs/Dialogs.h"
 #include "Framework/Docking/LayoutExtender.h"
+#include "ToolMenus.h"
 
 #define LOCTEXT_NAMESPACE "WorkspaceEditor"
 
@@ -45,6 +47,7 @@ const FName WorkspaceAppIdentifier("WorkspaceEditor");
 FWorkspaceEditor::FWorkspaceEditor(UWorkspaceAssetEditor* InOwningAssetEditor) : IWorkspaceEditor(InOwningAssetEditor)
 {
 	Workspace = Cast<UWorkspaceAssetEditor>(OwningAssetEditor)->GetObjectToEdit();
+	bCheckDirtyOnAssetSave = true;
 }
 
 void FWorkspaceEditor::CreateWidgets()
@@ -248,7 +251,7 @@ void FWorkspaceEditor::OpenObjects(TConstArrayView<UObject*> InObjects)
 {
 	for(const UObject* Object : InObjects)
 	{
-		OpenDocument(Object, FDocumentTracker::EOpenDocumentCause::NavigatingCurrentDocument);
+		OpenDocument(Object, InObjects.Num() == 1 ? FDocumentTracker::EOpenDocumentCause::NavigatingCurrentDocument : FDocumentTracker::EOpenDocumentCause::OpenNewDocument);
 	}
 }
 
@@ -304,13 +307,20 @@ void FWorkspaceEditor::BindCommands()
 	ToolkitCommands->MapAction
 	( 
 		Commands.NavigateBackward,
-		FExecuteAction::CreateRaw( this, &FWorkspaceEditor::NavigateBack )
+		FExecuteAction::CreateRaw(this, &FWorkspaceEditor::NavigateBack)
 	);
 	
 	ToolkitCommands->MapAction
 	( 
 		Commands.NavigateForward,
-		FExecuteAction::CreateRaw( this, &FWorkspaceEditor::NavigateForward )
+		FExecuteAction::CreateRaw(this, &FWorkspaceEditor::NavigateForward)
+	);
+
+	ToolkitCommands->MapAction
+	(
+		Commands.SaveAssetEntries,
+		FExecuteAction::CreateRaw(this, &FWorkspaceEditor::SaveAssetEntries),
+		FCanExecuteAction::CreateRaw(this, &FWorkspaceEditor::AreAssetEntriesModified)
 	);
 }
 
@@ -321,7 +331,11 @@ void FWorkspaceEditor::ExtendMenu()
 
 void FWorkspaceEditor::ExtendToolbar()
 {
-	
+	if (UToolMenu* Menu = UToolMenus::Get()->ExtendMenu(GetToolMenuToolbarName()))
+	{
+		FToolMenuSection& WorkspaceOperationsSection = Menu->AddSection("WorkspaceOperations");
+		WorkspaceOperationsSection.AddEntry(FToolMenuEntry::InitToolBarButton(FWorkspaceAssetEditorCommands::Get().SaveAssetEntries, FText::GetEmpty(), TAttribute<FText>(), FSlateIcon(FAppStyle::GetAppStyleSetName(), "MainFrame.SaveAll")));
+	}
 }
 
 void FWorkspaceEditor::RegisterTabSpawners(const TSharedRef<FTabManager>& InTabManager)
@@ -392,12 +406,14 @@ void FWorkspaceEditor::SaveAsset_Execute()
 	if(AssetPath.StartsWith(TEXT("/Temp/Untitled")))
 	{
 		// Ensure we do not also 'save as' other externally linked assets at this point
-		TGuardValue<bool> SaveWorkspaceOnly(bSavingWorkspaceOnly, true);
-
-		SaveAssetAs_Execute();
+		{
+			TGuardValue<bool> SavingTransientWorkspace(bSavingTransientWorkspace, true);
+			SaveAssetAs_Execute();
+		}
 	}
 	else
 	{
+		TGuardValue<bool> SavingWorkspace(bSavingWorkspace, true);
 		FBaseAssetToolkit::SaveAsset_Execute();
 	}
 }
@@ -418,16 +434,12 @@ bool FWorkspaceEditor::InEditingMode() const
 
 void FWorkspaceEditor::GetSaveableObjects(TArray<UObject*>& OutObjects) const
 {
-	// Base class will pick up edited object
-	FBaseAssetToolkit::GetSaveableObjects(OutObjects);
-
-	for (UObject* Object : GetEditingObjects())
+	if (bSavingWorkspace || bSavingTransientWorkspace)
 	{
-		// Get external objects too
-		FExternalPackageHelper::GetExternalSaveableObjects(Object, OutObjects);	
+		FBaseAssetToolkit::GetSaveableObjects(OutObjects);
 	}
 	
-	if(!bSavingWorkspaceOnly)
+	if(bSavingAssetEntries)
 	{
 		for(const UWorkspaceAssetEntry* Entry : Workspace->AssetEntries)
 		{
@@ -435,14 +447,32 @@ void FWorkspaceEditor::GetSaveableObjects(TArray<UObject*>& OutObjects) const
 			{
 				// Add object referenced by workspace
 				OutObjects.Add(Asset);
-
-				// Get external objects too
-				FExternalPackageHelper::GetExternalSaveableObjects(Asset, OutObjects);
 			}
 		}
 	}
 }
 
+bool FWorkspaceEditor::CanSaveAsset() const
+{
+	bool bDirtyState = false;
+	for (UObject* EditingObject : GetEditingObjects())
+	{
+		const UPackage* Package = EditingObject->GetOutermost();
+		if (Package->IsDirty() || Package->GetExternalPackages().ContainsByPredicate([](const UPackage* ExternalPackage)  { return ExternalPackage && ExternalPackage->IsDirty(); }))
+		{
+			bDirtyState = true;
+			break;
+		}
+	}
+
+	return bDirtyState;
+}
+
+FText FWorkspaceEditor::GetTabSuffix() const
+{
+	return CanSaveAsset() ? LOCTEXT("TabSuffixAsterix", "*") : FText::GetEmpty();
+}
+
 void FWorkspaceEditor::RecordDocumentState(const TInstancedStruct<FWorkspaceDocumentState>& InState) const
 {
 	UWorkspaceState* State = Workspace->GetState();
@@ -461,6 +491,29 @@ void FWorkspaceEditor::NavigateForward()
 	TSharedPtr<SDockTab> OpenedTab = DocumentManager->OpenDocument(Payload, FDocumentTracker::NavigateForwards);
 }
 
+void FWorkspaceEditor::SaveAssetEntries()
+{	
+	TGuardValue<bool> SavingAssetEntries(bSavingAssetEntries, true);
+	FBaseAssetToolkit::SaveAsset_Execute();
+}
+
+bool FWorkspaceEditor::AreAssetEntriesModified() const
+{
+	for(const UWorkspaceAssetEntry* Entry : Workspace->AssetEntries)
+	{
+		if(UObject* Asset = Entry->Asset.Get())
+		{
+			const UPackage* Package = Asset->GetOutermost();
+			if (Package->IsDirty() || Package->GetExternalPackages().ContainsByPredicate([](const UPackage* ExternalPackage)  { return ExternalPackage && ExternalPackage->IsDirty(); }))
+			{
+				return true;
+			}
+		}
+	}
+	
+	return false;
+}
+
 bool FWorkspaceEditor::OnRequestClose(EAssetEditorCloseReason InCloseReason)
 {
 	TGuardValue<bool> ClosingDown(bClosingDown, true);
@@ -472,12 +525,26 @@ bool FWorkspaceEditor::OnRequestClose(EAssetEditorCloseReason InCloseReason)
 	};
 
 	// Give the user opportunity to save temp workspaces
-	if(RequiresSave() && !bSavingWorkspaceOnly)
+	if(RequiresSave() && !bSavingTransientWorkspace)
 	{
-		// Ensure we dont also 'save as' other externally linked assets at this point
-		TGuardValue<bool> SaveWorkspaceOnly(bSavingWorkspaceOnly, true);
+		// Prompt whether to save or not, this can be skipped to become a never-ask-nor-save
+		FSuppressableWarningDialog::FSetupInfo Info(
+		LOCTEXT("SavingTransientWorkspaceAssetMessage", "Asset was opened in a temporary Workspace, do you want to save it?"),
+				LOCTEXT("SavingTransientWorkspaceAssetTitle", "Save temporary Workspace"), "SaveTemporaryWorkspacesPrompt");
+
+		Info.DialogMode = FSuppressableWarningDialog::EMode::PersistUserResponse;
+		Info.ConfirmText =LOCTEXT("SavingTransientWorkspaceAssetYes", "Yes");
+		Info.CancelText = LOCTEXT("SavingTransientWorkspaceAssetNo", "No");
+
+		FSuppressableWarningDialog SaveWorkspace(Info);
+		FSuppressableWarningDialog::EResult Result = SaveWorkspace.ShowModal();
 
-		SaveAssetAs_Execute();
+		if(Result == FSuppressableWarningDialog::EResult::Confirm)
+		{
+			// Ensure we dont also 'save as' other externally linked assets at this point
+			TGuardValue<bool> SaveWorkspaceOnly(bSavingTransientWorkspace, true);
+			SaveAssetAs_Execute();
+		}
 	}
 
 	return true;
diff --git a/Engine/Plugins/Experimental/Workspace/Source/WorkspaceEditor/Private/WorkspaceEditor.h b/Engine/Plugins/Experimental/Workspace/Source/WorkspaceEditor/Private/WorkspaceEditor.h
--- a/Engine/Plugins/Experimental/Workspace/Source/WorkspaceEditor/Private/WorkspaceEditor.h
+++ b/Engine/Plugins/Experimental/Workspace/Source/WorkspaceEditor/Private/WorkspaceEditor.h
@@ -71,6 +71,9 @@ private:
 
 	// FAssetEditorToolkit interface
 	virtual void GetSaveableObjects(TArray<UObject*>& OutObjects) const override;
+	virtual bool CanSaveAsset() const override;
+	virtual FText GetTabSuffix() const override;
+	virtual bool IsFindInContentBrowserButtonVisible() const { return false; }
 
 	// FGCObject interface
 	virtual void AddReferencedObjects(FReferenceCollector& Collector) override
@@ -113,6 +116,9 @@ private:
 	void NavigateBack();
 	void NavigateForward();
 
+	void SaveAssetEntries();
+	bool AreAssetEntriesModified() const;
+
 	/** The asset being edited */
 	TObjectPtr<UWorkspace> Workspace = nullptr;
 
@@ -124,7 +130,9 @@ private:
 	/** Tabs to be registered into the Workspace */
 	FWorkflowAllowedTabSet TabFactories;
 
-	bool bSavingWorkspaceOnly = false;
+	bool bSavingTransientWorkspace = false;
+	bool bSavingWorkspace = false;
+	bool bSavingAssetEntries = false;
 	bool bClosingDown = false;
 
 	FGlobalSelectionId LastGlobalSelectionId = nullptr;
diff --git a/Engine/Plugins/Experimental/Workspace/Source/WorkspaceEditor/Private/WorkspaceEditorCommands.cpp b/Engine/Plugins/Experimental/Workspace/Source/WorkspaceEditor/Private/WorkspaceEditorCommands.cpp
--- a/Engine/Plugins/Experimental/Workspace/Source/WorkspaceEditor/Private/WorkspaceEditorCommands.cpp
+++ b/Engine/Plugins/Experimental/Workspace/Source/WorkspaceEditor/Private/WorkspaceEditorCommands.cpp
@@ -27,6 +27,7 @@ namespace UE::Workspace
 			EUserInterfaceActionType::Button, FInputChord(EKeys::ThumbMouseButton), FInputChord(EKeys::Left, EModifierKey::Alt));
 		UI_COMMAND(NavigateForward, "Navigate Forward", "Moves forwards to previous location",
 			EUserInterfaceActionType::Button, FInputChord(EKeys::ThumbMouseButton2), FInputChord(EKeys::Right, EModifierKey::Alt));
+		UI_COMMAND(SaveAssetEntries, "Save Assets", "Saves all loaded and modified assets in the workspace", EUserInterfaceActionType::Button, FInputChord());
 	}
 
 }  // namespace UE::Workspace
diff --git a/Engine/Plugins/Experimental/Workspace/Source/WorkspaceEditor/Private/WorkspaceEditorCommands.h b/Engine/Plugins/Experimental/Workspace/Source/WorkspaceEditor/Private/WorkspaceEditorCommands.h
--- a/Engine/Plugins/Experimental/Workspace/Source/WorkspaceEditor/Private/WorkspaceEditorCommands.h
+++ b/Engine/Plugins/Experimental/Workspace/Source/WorkspaceEditor/Private/WorkspaceEditorCommands.h
@@ -16,6 +16,7 @@ namespace UE::Workspace
 
 		TSharedPtr<FUICommandInfo> NavigateBackward;		
 		TSharedPtr<FUICommandInfo> NavigateForward;
+		TSharedPtr<FUICommandInfo> SaveAssetEntries;
 	};
 
 }  // namespace UE::Workspace
diff --git a/Engine/Plugins/Experimental/Workspace/Source/WorkspaceEditor/Public/IWorkspaceEditorModule.h b/Engine/Plugins/Experimental/Workspace/Source/WorkspaceEditor/Public/IWorkspaceEditorModule.h
--- a/Engine/Plugins/Experimental/Workspace/Source/WorkspaceEditor/Public/IWorkspaceEditorModule.h
+++ b/Engine/Plugins/Experimental/Workspace/Source/WorkspaceEditor/Public/IWorkspaceEditorModule.h
@@ -60,9 +60,13 @@ struct FWorkspaceBreadcrumb : TSharedFromThis<FWorkspaceBreadcrumb>
 {
 	using FOnGetBreadcrumbLabel = TDelegate<TAttribute<FText>()>;
 	using FOnBreadcrumbClicked = TDelegate<void()>;
+	using FCanSaveBreadcrumb = TDelegate<bool()>;
+	using FOnSaveBreadcrumb = TDelegate<void()>;
 	
 	FOnGetBreadcrumbLabel OnGetLabel;
 	FOnBreadcrumbClicked OnClicked;
+	FCanSaveBreadcrumb CanSave;
+	FOnSaveBreadcrumb OnSave;
 };
 
 using FOnMakeDocumentWidget = TDelegate<TSharedRef<SWidget>(const FWorkspaceEditorContext&)>;
