Change outliner item export to use full path as its identifier instead of the individual property bits. Path consistes out of asset-path, parent identifier(s) and identifier segments. Which is now used to create a consistent unique hash, solving collisions prompted by collapse nodes etc.

diff --git a/Engine/Plugins/Experimental/AnimNext/Source/AnimNextEditor/Private/Graph/AnimNextGraphItemDetails.cpp b/Engine/Plugins/Experimental/AnimNext/Source/AnimNextEditor/Private/Graph/AnimNextGraphItemDetails.cpp
--- a/Engine/Plugins/Experimental/AnimNext/Source/AnimNextEditor/Private/Graph/AnimNextGraphItemDetails.cpp
+++ b/Engine/Plugins/Experimental/AnimNext/Source/AnimNextEditor/Private/Graph/AnimNextGraphItemDetails.cpp
@@ -25,7 +25,7 @@ void FAnimNextGraphItemDetails::HandleDoubleClick(const FToolMenuContext& ToolMe
 	{
 		if(const TSharedPtr<UE::Workspace::IWorkspaceEditor> WorkspaceEditor = StaticCastSharedPtr<UE::Workspace::IWorkspaceEditor>(AssetEditorContext->Toolkit.Pin()))
 		{
-			const TInstancedStruct<FWorkspaceOutlinerItemData>& Data = WorkspaceItemContext->SelectedExports[0].Data;
+			const TInstancedStruct<FWorkspaceOutlinerItemData>& Data = WorkspaceItemContext->SelectedExports[0].GetData();
 			if (Data.IsValid() && Data.GetScriptStruct() == FAnimNextGraphOutlinerData::StaticStruct())
 			{
 				const FAnimNextGraphOutlinerData& GraphData = Data.Get<FAnimNextGraphOutlinerData>();							
@@ -49,7 +49,7 @@ void FAnimNextGraphItemDetails::HandleDoubleClick(const FToolMenuContext& ToolMe
 
 UPackage* FAnimNextGraphItemDetails::GetPackage(const FWorkspaceOutlinerItemExport& Export) const 
 {
-	const TInstancedStruct<FWorkspaceOutlinerItemData>& Data = Export.Data;
+	const TInstancedStruct<FWorkspaceOutlinerItemData>& Data = Export.GetData();
 	if (Data.IsValid() && Data.GetScriptStruct() == FAnimNextGraphOutlinerData::StaticStruct())
 	{
 		const FAnimNextGraphOutlinerData& GraphData = Data.Get<FAnimNextGraphOutlinerData>();
@@ -84,7 +84,7 @@ void FAnimNextGraphItemDetails::RegisterToolMenuExtensions()
 					TArray<FWorkspaceOutlinerItemExport> GraphExports;
 					Algo::TransformIf(WorkspaceItemContext->SelectedExports, GraphExports, [](const FWorkspaceOutlinerItemExport& Export)
 					{
-						return Export.Data.IsValid() && Export.Data.GetScriptStruct() == FAnimNextGraphOutlinerData::StaticStruct();
+						return Export.GetData().IsValid() && Export.GetData().GetScriptStruct() == FAnimNextGraphOutlinerData::StaticStruct();
 					},
 					[](const FWorkspaceOutlinerItemExport& Export)
 					{
@@ -107,7 +107,7 @@ void FAnimNextGraphItemDetails::RegisterToolMenuExtensions()
 									TArray<UObject*> ObjectsToOpen;
 									for (const FWorkspaceOutlinerItemExport& Export : GraphExports)
 									{
-										const FAnimNextGraphOutlinerData& GraphData = Export.Data.Get<FAnimNextGraphOutlinerData>();
+										const FAnimNextGraphOutlinerData& GraphData = Export.GetData().Get<FAnimNextGraphOutlinerData>();
 										if (GraphData.GraphInterface)
 										{
 											if (URigVMGraph* RigVMGraph = GraphData.GraphInterface->GetRigVMGraph())
diff --git a/Engine/Plugins/Experimental/AnimNext/Source/AnimNextUncookedOnly/Private/AnimNextUncookedOnlyModule.cpp b/Engine/Plugins/Experimental/AnimNext/Source/AnimNextUncookedOnly/Private/AnimNextUncookedOnlyModule.cpp
--- a/Engine/Plugins/Experimental/AnimNext/Source/AnimNextUncookedOnly/Private/AnimNextUncookedOnlyModule.cpp
+++ b/Engine/Plugins/Experimental/AnimNext/Source/AnimNextUncookedOnly/Private/AnimNextUncookedOnlyModule.cpp
@@ -40,13 +40,9 @@ void FModule::StartupModule()
 
 		FWorkspaceOutlinerItemExports OutlinerExports;
 		{
-			FWorkspaceOutlinerItemExport& RootAssetExport = OutlinerExports.Exports.AddDefaulted_GetRef();
-			RootAssetExport.Identifier = InSchedule->GetFName();
-			RootAssetExport.ParentIdentifier = NAME_None;
-			RootAssetExport.AssetPath = InSchedule;
-
-			RootAssetExport.Data.InitializeAsScriptStruct(FAnimNextSchedulerData::StaticStruct());
-			FAnimNextSchedulerData& AssetData = RootAssetExport.Data.GetMutable<FAnimNextSchedulerData>();
+			FWorkspaceOutlinerItemExport& RootAssetExport = OutlinerExports.Exports.Add_GetRef(FWorkspaceOutlinerItemExport(InSchedule->GetFName(), InSchedule));
+			RootAssetExport.GetData().InitializeAsScriptStruct(FAnimNextSchedulerData::StaticStruct());
+			FAnimNextSchedulerData& AssetData = RootAssetExport.GetData().GetMutable<FAnimNextSchedulerData>();
 
 			{
 				FString TagValue;
diff --git a/Engine/Plugins/Experimental/AnimNext/Source/AnimNextUncookedOnly/Private/Module/AnimNextModuleWorkspaceAssetUserData.cpp b/Engine/Plugins/Experimental/AnimNext/Source/AnimNextUncookedOnly/Private/Module/AnimNextModuleWorkspaceAssetUserData.cpp
--- a/Engine/Plugins/Experimental/AnimNext/Source/AnimNextUncookedOnly/Private/Module/AnimNextModuleWorkspaceAssetUserData.cpp
+++ b/Engine/Plugins/Experimental/AnimNext/Source/AnimNextUncookedOnly/Private/Module/AnimNextModuleWorkspaceAssetUserData.cpp
@@ -16,12 +16,8 @@ void UAnimNextModuleWorkspaceAssetUserData::GetAssetRegistryTags(FAssetRegistryT
 		const UAnimNextModule* ModuleOuter = CastChecked<UAnimNextModule>(GetOuter());
 		const UAnimNextModule_EditorData* GraphEditorData = UE::AnimNext::UncookedOnly::FUtils::GetEditorData(ModuleOuter);	
 		{
-			FWorkspaceOutlinerItemExport& RootAssetExport = Exports.Exports.AddDefaulted_GetRef();	
-			RootAssetExport.Identifier = ModuleOuter->GetFName();
-			RootAssetExport.ParentIdentifier = NAME_None;
-			RootAssetExport.AssetPath = ModuleOuter;
-
-			RootAssetExport.Data.InitializeAsScriptStruct(FAnimNextGraphAssetOutlinerData::StaticStruct());
+			FWorkspaceOutlinerItemExport& RootAssetExport = Exports.Exports.Add_GetRef(FWorkspaceOutlinerItemExport(ModuleOuter->GetFName(), ModuleOuter));
+			RootAssetExport.GetData().InitializeAsScriptStruct(FAnimNextGraphAssetOutlinerData::StaticStruct());
 		}
 	
 		UE::AnimNext::UncookedOnly::FUtils::GetAssetOutlinerItems(GraphEditorData, Exports);
diff --git a/Engine/Plugins/Experimental/AnimNext/Source/AnimNextUncookedOnly/Private/UncookedOnlyUtils.cpp b/Engine/Plugins/Experimental/AnimNext/Source/AnimNextUncookedOnly/Private/UncookedOnlyUtils.cpp
--- a/Engine/Plugins/Experimental/AnimNext/Source/AnimNextUncookedOnly/Private/UncookedOnlyUtils.cpp
+++ b/Engine/Plugins/Experimental/AnimNext/Source/AnimNextUncookedOnly/Private/UncookedOnlyUtils.cpp
@@ -1655,6 +1655,8 @@ void FUtils::GetBlueprintParameters(const UBlueprint* InBlueprint, TSet<FAnimNex
 
 void FUtils::GetAssetOutlinerItems(const UAnimNextRigVMAssetEditorData* EditorData, FWorkspaceOutlinerItemExports& OutExports)
 {
+	FWorkspaceOutlinerItemExport AssetIdentifier = FWorkspaceOutlinerItemExport(EditorData->GetOuter()->GetFName(), EditorData->GetOuter());
+	
 	constexpr bool bExportParametersAsOutlinerItems = false;
 	if (bExportParametersAsOutlinerItems)
 	{
@@ -1666,13 +1668,10 @@ void FUtils::GetAssetOutlinerItems(const UAnimNextRigVMAssetEditorData* EditorDa
 		{
 			if (!ParameterNames.Contains(Entry.Name))
 			{
-				FWorkspaceOutlinerItemExport& ParameterExport = OutExports.Exports.AddDefaulted_GetRef();	
-				ParameterExport.Identifier = Entry.Name;						
-				ParameterExport.ParentIdentifier = EditorData->GetOuter()->GetFName();
-				ParameterExport.AssetPath = EditorData->GetOuter();
+				FWorkspaceOutlinerItemExport& ParameterExport = OutExports.Exports.Add_GetRef(FWorkspaceOutlinerItemExport(Entry.Name, AssetIdentifier));	
 
-				ParameterExport.Data.InitializeAsScriptStruct(FAnimNextParameterOutlinerData::StaticStruct());
-				FAnimNextParameterOutlinerData& AssetData = ParameterExport.Data.GetMutable<FAnimNextParameterOutlinerData>();
+				ParameterExport.GetData().InitializeAsScriptStruct(FAnimNextParameterOutlinerData::StaticStruct());
+				FAnimNextParameterOutlinerData& AssetData = ParameterExport.GetData().GetMutable<FAnimNextParameterOutlinerData>();
 				AssetData.Type = Entry.Type;
 
 				ParameterNames.Add(Entry.Name);
@@ -1684,13 +1683,10 @@ void FUtils::GetAssetOutlinerItems(const UAnimNextRigVMAssetEditorData* EditorDa
 	{
 		if(const IAnimNextRigVMGraphInterface* GraphInterface = Cast<IAnimNextRigVMGraphInterface>(Entry))
 		{
-			FWorkspaceOutlinerItemExport& Export = OutExports.Exports.AddDefaulted_GetRef();
-			Export.Identifier = Entry->GetEntryName();
-			Export.ParentIdentifier = EditorData->GetOuter()->GetFName();					
-			Export.AssetPath = EditorData->GetOuter();
+			FWorkspaceOutlinerItemExport& Export = OutExports.Exports.Add_GetRef(FWorkspaceOutlinerItemExport(Entry->GetEntryName(), AssetIdentifier));
 
-			Export.Data.InitializeAsScriptStruct(FAnimNextGraphOutlinerData::StaticStruct());
-			FAnimNextGraphOutlinerData& GraphData = Export.Data.GetMutable<FAnimNextGraphOutlinerData>();
+			Export.GetData().InitializeAsScriptStruct(FAnimNextGraphOutlinerData::StaticStruct());
+			FAnimNextGraphOutlinerData& GraphData = Export.GetData().GetMutable<FAnimNextGraphOutlinerData>();
 			GraphData.GraphInterface = GraphInterface->_getUObject();
 		}
 	}
diff --git a/Engine/Plugins/Experimental/Workspace/Source/WorkspaceEditor/Private/Outliner/WorkspaceOutlinerColumns.cpp b/Engine/Plugins/Experimental/Workspace/Source/WorkspaceEditor/Private/Outliner/WorkspaceOutlinerColumns.cpp
--- a/Engine/Plugins/Experimental/Workspace/Source/WorkspaceEditor/Private/Outliner/WorkspaceOutlinerColumns.cpp
+++ b/Engine/Plugins/Experimental/Workspace/Source/WorkspaceEditor/Private/Outliner/WorkspaceOutlinerColumns.cpp
@@ -59,7 +59,7 @@ namespace UE::Workspace
 				
 				if (!WeakPackage.IsValid())
 				{
-					if(const UObject* Object = TreeItem->Export.AssetPath.ResolveObject())
+					if(const UObject* Object = TreeItem->Export.GetAssetPath().ResolveObject())
 					{
 						WeakPackage = Object->GetPackage();
 					}
diff --git a/Engine/Plugins/Experimental/Workspace/Source/WorkspaceEditor/Private/Outliner/WorkspaceOutlinerHierarchy.cpp b/Engine/Plugins/Experimental/Workspace/Source/WorkspaceEditor/Private/Outliner/WorkspaceOutlinerHierarchy.cpp
--- a/Engine/Plugins/Experimental/Workspace/Source/WorkspaceEditor/Private/Outliner/WorkspaceOutlinerHierarchy.cpp
+++ b/Engine/Plugins/Experimental/Workspace/Source/WorkspaceEditor/Private/Outliner/WorkspaceOutlinerHierarchy.cpp
@@ -43,10 +43,10 @@ namespace UE::Workspace
 	{
 		if (const FWorkspaceOutlinerTreeItem* TreeItem = Item.CastTo<FWorkspaceOutlinerTreeItem>())
 		{
-			const FName ParentIdentifier = TreeItem->Export.ParentIdentifier;
-			if (ParentIdentifier != NAME_None)
+			const uint32 ParentHash = TreeItem->Export.GetParentHash();
+			if (ParentHash != INDEX_NONE)
 			{
-				if (const FSceneOutlinerTreeItemPtr* ParentItem = Items.Find(HashCombine(GetTypeHash(TreeItem->Export.AssetPath),GetTypeHash(ParentIdentifier))))
+				if (const FSceneOutlinerTreeItemPtr* ParentItem = Items.Find(ParentHash))
 				{
 					return *ParentItem;
 				}
@@ -57,17 +57,18 @@ namespace UE::Workspace
 						TArray<FAssetData> AssetDataEntries; 
 						Workspace->GetAssetDataEntries(AssetDataEntries);
 
-						if(const FAssetData* AssetDataPtr = AssetDataEntries.FindByPredicate([AssetPath = TreeItem->Export.AssetPath](const FAssetData& AssetData) { return AssetData.GetSoftObjectPath() == AssetPath; }))
+						if(const FAssetData* AssetDataPtr = AssetDataEntries.FindByPredicate([AssetPath = TreeItem->Export.GetAssetPath()](const FAssetData& AssetData) { return AssetData.GetSoftObjectPath() == AssetPath; }))
 						{
 							FString TagValue;
-							if((*AssetDataPtr).GetTagValue(UE::Workspace::ExportsWorkspaceItemsRegistryTag, TagValue))
+							if(AssetDataPtr->GetTagValue(UE::Workspace::ExportsWorkspaceItemsRegistryTag, TagValue))
 							{
 								FWorkspaceOutlinerItemExports Exports;
 								FWorkspaceOutlinerItemExports::StaticStruct()->ImportText(*TagValue, &Exports, nullptr, 0, nullptr, FWorkspaceOutlinerItemExports::StaticStruct()->GetName());
 								
+								const FName ParentIdentifier = TreeItem->Export.GetParentIdentifier();
 								if (const FWorkspaceOutlinerItemExport* ExportPtr = Exports.Exports.FindByPredicate([ParentIdentifier](const FWorkspaceOutlinerItemExport& ItemExport)
 								{
-									return ItemExport.Identifier == ParentIdentifier;
+									return ItemExport.GetIdentifier() == ParentIdentifier;
 								}))
 								{
 									Mode->CreateItemFor<FWorkspaceOutlinerTreeItem>(FWorkspaceOutlinerTreeItem::FItemData{*ExportPtr}, true);
diff --git a/Engine/Plugins/Experimental/Workspace/Source/WorkspaceEditor/Private/Outliner/WorkspaceOutlinerMode.cpp b/Engine/Plugins/Experimental/Workspace/Source/WorkspaceEditor/Private/Outliner/WorkspaceOutlinerMode.cpp
--- a/Engine/Plugins/Experimental/Workspace/Source/WorkspaceEditor/Private/Outliner/WorkspaceOutlinerMode.cpp
+++ b/Engine/Plugins/Experimental/Workspace/Source/WorkspaceEditor/Private/Outliner/WorkspaceOutlinerMode.cpp
@@ -68,7 +68,7 @@ TSharedPtr<SWidget> FWorkspaceOutlinerMode::CreateContextMenu()
 				{
 					const bool bSelectionContainsTopLevelAsset = MenuContext->SelectedExports.Num() && MenuContext->SelectedExports.ContainsByPredicate([](const FWorkspaceOutlinerItemExport& Export)
 					{
-						return Export.ParentIdentifier == NAME_None;
+						return Export.GetParentIdentifier() == NAME_None;
 					});
 
 					FToolMenuSection& AssetsSection = InMenu->AddSection("Assets", LOCTEXT("AssetSectionLabel", "Assets"));
@@ -85,9 +85,9 @@ TSharedPtr<SWidget> FWorkspaceOutlinerMode::CreateContextMenu()
 									TSet<FSoftObjectPath> AssetPaths;
 									for (const FWorkspaceOutlinerItemExport& ItemExport : SelectedExports)
 									{
-										if (ItemExport.ParentIdentifier == NAME_None)
+										if (ItemExport.GetParentIdentifier() == NAME_None)
 										{
-											AssetPaths.Add(ItemExport.AssetPath);	
+											AssetPaths.Add(ItemExport.GetAssetPath());	
 										}
 									}
 
@@ -110,9 +110,9 @@ TSharedPtr<SWidget> FWorkspaceOutlinerMode::CreateContextMenu()
 									TSet<FSoftObjectPath> AssetPaths;
 									for (const FWorkspaceOutlinerItemExport& ItemExport : SelectedExports)
 									{
-										if (ItemExport.ParentIdentifier == NAME_None)
+										if (ItemExport.GetParentIdentifier() == NAME_None)
 										{
-											AssetPaths.Add(ItemExport.AssetPath);	
+											AssetPaths.Add(ItemExport.GetAssetPath());	
 										}
 									}
 
@@ -156,9 +156,9 @@ TSharedPtr<SWidget> FWorkspaceOutlinerMode::CreateContextMenu()
 											SavablePackages.AddUnique(Package);								
 										}										
 									}
-									else if (ItemExport.ParentIdentifier == NAME_None)
+									else if (ItemExport.GetParentIdentifier() == NAME_None)
 									{
-										if (UPackage* Package = FindPackage(nullptr, *ItemExport.AssetPath.GetLongPackageName()))
+										if (UPackage* Package = FindPackage(nullptr, *ItemExport.GetAssetPath().GetLongPackageName()))
 										{
 											if (IsPackageDirty(Package))
 											{
@@ -186,9 +186,9 @@ TSharedPtr<SWidget> FWorkspaceOutlinerMode::CreateContextMenu()
 											return true;
 										}
 									}
-									else if (ItemExport.ParentIdentifier == NAME_None)
+									else if (ItemExport.GetParentIdentifier() == NAME_None)
 									{
-										if (const UPackage* Package = FindPackage(nullptr, *ItemExport.AssetPath.GetLongPackageName()))
+										if (const UPackage* Package = FindPackage(nullptr, *ItemExport.GetAssetPath().GetLongPackageName()))
 										{
 											if (IsPackageDirty(Package))
 											{
@@ -276,9 +276,9 @@ void FWorkspaceOutlinerMode::HandleItemSelection(const FSceneOutlinerItemSelecti
 		{
 			if(FWorkspaceOutlinerTreeItem* TreeItem = SelectedItems[0]->CastTo<FWorkspaceOutlinerTreeItem>())
 			{
-				if(TreeItem->Export.Data.IsValid())
+				if(TreeItem->Export.GetData().IsValid())
 				{
-					TSharedPtr<FStructOnScope> ExportDataView = MakeShared<FStructOnScope>(TreeItem->Export.Data.GetScriptStruct(), TreeItem->Export.Data.GetMutableMemory());
+					TSharedPtr<FStructOnScope> ExportDataView = MakeShared<FStructOnScope>(TreeItem->Export.GetData().GetScriptStruct(), TreeItem->Export.GetData().GetMutableMemory());
 					// TODO JDB handle struct selections
 					//SharedWorkspaceEditor->SetDetailsStruct(ExportDataView);
 				}
@@ -319,11 +319,11 @@ void FWorkspaceOutlinerMode::OpenItems(TArrayView<const FSceneOutlinerTreeItemPt
 	{
 		if (const FWorkspaceOutlinerTreeItem* TreeItem = Item->CastTo<FWorkspaceOutlinerTreeItem>())
 		{
-			if(TreeItem->Export.ParentIdentifier == NAME_None)
+			if(TreeItem->Export.GetParentIdentifier() == NAME_None)
 			{
 				if (const TSharedPtr<UE::Workspace::IWorkspaceEditor> SharedWorkspaceEditor = StaticCastSharedPtr<UE::Workspace::IWorkspaceEditor>(WeakWorkspaceEditor.Pin()))
 				{			
-					SharedWorkspaceEditor->OpenAssets({TreeItem->Export.AssetPath.TryLoad()});
+					SharedWorkspaceEditor->OpenAssets({TreeItem->Export.GetAssetPath().TryLoad()});
 				}
 			}
 			else
@@ -350,9 +350,9 @@ void FWorkspaceOutlinerMode::DeleteItems(TArrayView<const FSceneOutlinerTreeItem
 	{
 		if (const FWorkspaceOutlinerTreeItem* TreeItem = Item->CastTo<FWorkspaceOutlinerTreeItem>())
 		{
-			if(TreeItem->Export.ParentIdentifier == NAME_None)
+			if(TreeItem->Export.GetParentIdentifier() == NAME_None)
 			{
-				Filter.SoftObjectPaths.AddUnique(TreeItem->Export.AssetPath);
+				Filter.SoftObjectPaths.AddUnique(TreeItem->Export.GetAssetPath());
 			}
 		}		
 	}
diff --git a/Engine/Plugins/Experimental/Workspace/Source/WorkspaceEditor/Private/Outliner/WorkspaceOutlinerTreeItem.cpp b/Engine/Plugins/Experimental/Workspace/Source/WorkspaceEditor/Private/Outliner/WorkspaceOutlinerTreeItem.cpp
--- a/Engine/Plugins/Experimental/Workspace/Source/WorkspaceEditor/Private/Outliner/WorkspaceOutlinerTreeItem.cpp
+++ b/Engine/Plugins/Experimental/Workspace/Source/WorkspaceEditor/Private/Outliner/WorkspaceOutlinerTreeItem.cpp
@@ -97,7 +97,7 @@ namespace UE::Workspace
 
 	bool FWorkspaceOutlinerTreeItem::IsValid() const
 	{
-		return Export.Identifier.IsValid();
+		return Export.GetIdentifier().IsValid();
 	}
 
 	FSceneOutlinerTreeItemID FWorkspaceOutlinerTreeItem::GetID() const
@@ -107,7 +107,7 @@ namespace UE::Workspace
 
 	FString FWorkspaceOutlinerTreeItem::GetDisplayString() const
 	{
-		return Export.Identifier.ToString();
+		return Export.GetIdentifier().ToString();
 	}
 
 	TSharedRef<SWidget> FWorkspaceOutlinerTreeItem::GenerateLabelWidget(ISceneOutliner& Outliner, const STableRow<FSceneOutlinerTreeItemPtr>& InRow)
@@ -121,9 +121,9 @@ namespace UE::Workspace
 		{
 			return SharedFactory->GetPackage(Export)->GetName();
 		}
-		else if (Export.ParentIdentifier == NAME_None && Export.AssetPath.IsValid())
+		else if (Export.GetParentIdentifier() == NAME_None && Export.GetAssetPath().IsValid())
 		{
-			return Export.AssetPath.GetLongPackageName();
+			return Export.GetAssetPath().GetLongPackageName();
 		}		
 		
 		return ISceneOutlinerTreeItem::GetPackageName();
diff --git a/Engine/Plugins/Experimental/Workspace/Source/WorkspaceEditor/Public/IWorkspaceOutlinerItemDetails.h b/Engine/Plugins/Experimental/Workspace/Source/WorkspaceEditor/Public/IWorkspaceOutlinerItemDetails.h
--- a/Engine/Plugins/Experimental/Workspace/Source/WorkspaceEditor/Public/IWorkspaceOutlinerItemDetails.h
+++ b/Engine/Plugins/Experimental/Workspace/Source/WorkspaceEditor/Public/IWorkspaceOutlinerItemDetails.h
@@ -14,7 +14,7 @@ namespace UE::Workspace
 typedef FName FOutlinerItemDetailsId;
 static FOutlinerItemDetailsId MakeOutlinerDetailsId(const FWorkspaceOutlinerItemExport& InExport)
 {
-    return InExport.Data.IsValid() ? InExport.Data.GetScriptStruct()->GetFName() : NAME_None;
+    return InExport.GetData().IsValid() ? InExport.GetData().GetScriptStruct()->GetFName() : NAME_None;
 }
 
 class IWorkspaceOutlinerItemDetails : public TSharedFromThis<IWorkspaceOutlinerItemDetails>
diff --git a/Engine/Plugins/Experimental/Workspace/Source/WorkspaceEditor/Public/WorkspaceAssetRegistryInfo.h b/Engine/Plugins/Experimental/Workspace/Source/WorkspaceEditor/Public/WorkspaceAssetRegistryInfo.h
--- a/Engine/Plugins/Experimental/Workspace/Source/WorkspaceEditor/Public/WorkspaceAssetRegistryInfo.h
+++ b/Engine/Plugins/Experimental/Workspace/Source/WorkspaceEditor/Public/WorkspaceAssetRegistryInfo.h
@@ -13,6 +13,63 @@ struct FWorkspaceOutlinerItemData
 	FWorkspaceOutlinerItemData() = default;
 };
 
+USTRUCT()
+struct FOutlinerItemPath
+{
+	GENERATED_BODY()
+	
+	friend struct FWorkspaceOutlinerItemExport;
+protected:
+	UPROPERTY()
+	TArray<FName> PathSegments;
+
+public:
+	static FOutlinerItemPath MakePath(const FSoftObjectPath& InSoftObjectPath)
+	{
+		FOutlinerItemPath Path;
+		Path.PathSegments.Add(*InSoftObjectPath.ToString());
+		return Path;
+	}
+
+	FOutlinerItemPath AppendSegment(const FName& InSegment) const
+	{
+		FOutlinerItemPath Path = *this;
+		Path.PathSegments.Add(InSegment);
+		return Path;
+	}
+
+	FOutlinerItemPath RemoveSegment() const
+	{
+		FOutlinerItemPath Path = *this;
+		if (Path.PathSegments.Num())
+		{
+			Path.PathSegments.Pop();
+		}			
+		return Path;
+	}
+
+	friend uint32 GetTypeHash(const FOutlinerItemPath& Path)
+	{
+		uint32 Hash = INDEX_NONE;
+
+		if (Path.PathSegments.Num() == 1)
+		{
+			Hash = GetTypeHash(Path.PathSegments[0]);
+		}
+		else if (Path.PathSegments.Num() > 1)
+		{
+			Hash = GetTypeHash(Path.PathSegments[0]);
+			for (int32 Index = 1; Index < Path.PathSegments.Num(); ++Index)
+			{
+				Hash = HashCombine(Hash, GetTypeHash(Path.PathSegments[Index]));				
+			}
+			return Hash;
+		}
+		
+		return Hash;
+	}
+};
+
 USTRUCT()
 struct FWorkspaceOutlinerItemExport
 {
@@ -20,21 +77,68 @@ struct FWorkspaceOutlinerItemExport
 
 	FWorkspaceOutlinerItemExport() = default;
 
-	UPROPERTY()
-	FName Identifier;
+	FWorkspaceOutlinerItemExport(const FName InIdentifier, const FSoftObjectPath& InObjectPath)
+	{
+		Path.PathSegments.Add(*InObjectPath.ToString());
+		Path.PathSegments.Add(InIdentifier);
+	}
+	
+	FWorkspaceOutlinerItemExport(const FName InIdentifier, const FWorkspaceOutlinerItemExport& InParent)
+	{
+		Path = InParent.Path.AppendSegment(InIdentifier);
+	}
 
+protected:
+	/** Full 'path' to item this instance represents, expected to take form of AssetPath follow by a set of identifier names */
 	UPROPERTY()
-	FName ParentIdentifier;
-
-	UPROPERTY()	
-	FSoftObjectPath AssetPath;
+	FOutlinerItemPath Path;
 
 	UPROPERTY()
 	TInstancedStruct<FWorkspaceOutlinerItemData> Data;
+public:
+	FName GetIdentifier() const
+	{
+		// Path needs atleast two segments to contain a valid identifier
+		if(Path.PathSegments.Num() > 1)
+		{
+			return Path.PathSegments.Last();	
+		}
+
+		return NAME_None;
+	}
+	
+	FName GetParentIdentifier() const
+	{
+		// Path needs atleast three segments to contain a valid _parent_ identifier
+		const int32 NumSegments = Path.PathSegments.Num();
+		if (NumSegments > 2)
+		{
+			return Path.PathSegments[FMath::Max(NumSegments - 2, 0)];
+		}
+
+		return NAME_None;
+	}
+	
+	FSoftObjectPath GetAssetPath() const
+	{
+		// Path needs atleast one segment to contain a (potentially) valid asset path
+		if(Path.PathSegments.Num() > 0)
+		{
+			return FSoftObjectPath(Path.PathSegments[0].ToString());	
+		}
+
+		return FSoftObjectPath();
+	}
+	
+	// Remove identifier segment to retrieve parent path hash
+	uint32 GetParentHash() const { return GetTypeHash(Path.RemoveSegment()); }
+	
+	const TInstancedStruct<FWorkspaceOutlinerItemData>& GetData() const { return Data; }
+	TInstancedStruct<FWorkspaceOutlinerItemData>& GetData() { return Data; }	
 
 	friend uint32 GetTypeHash(const FWorkspaceOutlinerItemExport& Export)
 	{
-		return HashCombine(GetTypeHash(Export.AssetPath), GetTypeHash(Export.Identifier));
+		return GetTypeHash(Export.Path);
 	}
 };
 
