Added concept of 'focussed asset' which keeps track of the currently (or last) focussed asset within the workspace editor

diff --git a/Engine/Plugins/Experimental/Workspace/Source/WorkspaceEditor/Private/AssetDocumentSummoner.cpp b/Engine/Plugins/Experimental/Workspace/Source/WorkspaceEditor/Private/AssetDocumentSummoner.cpp
--- a/Engine/Plugins/Experimental/Workspace/Source/WorkspaceEditor/Private/AssetDocumentSummoner.cpp
+++ b/Engine/Plugins/Experimental/Workspace/Source/WorkspaceEditor/Private/AssetDocumentSummoner.cpp
@@ -32,6 +32,11 @@ void FAssetDocumentSummoner::SetAllowedClassPaths(TConstArrayView<FTopLevelAsset
 
 void FAssetDocumentSummoner::OnTabActivated(TSharedPtr<SDockTab> Tab) const
 {
+	TSharedRef<SWorkspaceTabWrapper> TabWrapper = StaticCastSharedRef<SWorkspaceTabWrapper>(Tab->GetContent());
+	if (UObject* DocumentAsset = TabWrapper->GetDocumentObject().Get())
+	{
+		HostingAppPtr.Pin()->SetFocussedAsset(DocumentAsset);
+	}
 }
 
 void FAssetDocumentSummoner::OnTabBackgrounded(TSharedPtr<SDockTab> Tab) const
@@ -116,7 +121,7 @@ TAttribute<FText> FAssetDocumentSummoner::ConstructTabLabelSuffix(const FWorkflo
 		{
 			if(UObject* Object = WeakObject.Get())
 			{
-				if(Object->IsAsset() && Object->GetPackage()->IsDirty())
+				if(Object->GetPackage()->IsDirty())
 				{
 					return LOCTEXT("TabSuffixAsterisk", "*");
 				}
diff --git a/Engine/Plugins/Experimental/Workspace/Source/WorkspaceEditor/Private/SWorkspaceTabWrapper.cpp b/Engine/Plugins/Experimental/Workspace/Source/WorkspaceEditor/Private/SWorkspaceTabWrapper.cpp
--- a/Engine/Plugins/Experimental/Workspace/Source/WorkspaceEditor/Private/SWorkspaceTabWrapper.cpp
+++ b/Engine/Plugins/Experimental/Workspace/Source/WorkspaceEditor/Private/SWorkspaceTabWrapper.cpp
@@ -1,7 +1,10 @@
 // Copyright Epic Games, Inc. All Rights Reserved.
 
 #include "SWorkspaceTabWrapper.h"
+
+#include "EditorModeManager.h"
 #include "Widgets/Layout/SSeparator.h"
+#include "Widgets/Docking/SDockTab.h"
 #include "Modules/ModuleManager.h"
 #include "ToolMenus.h"
 
@@ -13,6 +16,17 @@ void SWorkspaceTabWrapper::Construct( const FArguments& InArgs, TSharedPtr<class
 	WeakWorkspaceEditor = InWorkspaceEditor;
 	WeakDocumentObject = InDocumentID;
 
+	if (InTabInfo->GetTab().IsValid())
+	{
+		InTabInfo->GetTab().Pin()->SetOnTabClosed(SDockTab::FOnTabClosedCallback::CreateLambda([WeakEditor=WeakWorkspaceEditor, WeakObject=WeakDocumentObject](TSharedRef<SDockTab>)
+		{
+			if (const TSharedPtr<UE::Workspace::FWorkspaceEditor> SharedWorkspaceEditor = WeakEditor.Pin())
+			{
+				SharedWorkspaceEditor->RemoveEditingObject(WeakObject.Get());
+			}
+		}));
+	}
+
 	// Set-up shared breadcrumb defaults JDB TODO figure out correct padding to align fake title with breadcrumbs
     const FMargin BreadcrumbTrailPadding = FMargin(4.f, 2.f);
     const FSlateBrush* BreadcrumbButtonImage = FAppStyle::GetBrush("BreadcrumbTrail.Delimiter");
diff --git a/Engine/Plugins/Experimental/Workspace/Source/WorkspaceEditor/Private/SWorkspaceTabWrapper.h b/Engine/Plugins/Experimental/Workspace/Source/WorkspaceEditor/Private/SWorkspaceTabWrapper.h
--- a/Engine/Plugins/Experimental/Workspace/Source/WorkspaceEditor/Private/SWorkspaceTabWrapper.h
+++ b/Engine/Plugins/Experimental/Workspace/Source/WorkspaceEditor/Private/SWorkspaceTabWrapper.h
@@ -20,6 +20,7 @@ public:
 	
 	void Construct(const FArguments& InArgs, TSharedPtr<class FTabInfo> InTabInfo, TSharedPtr<UE::Workspace::FWorkspaceEditor> InWorkspaceEditor, UObject* InDocumentID);
 	TSharedRef<SWidget> GetContent() const { return Content.ToSharedRef(); }
+	TWeakObjectPtr<UObject> GetDocumentObject() const { return WeakDocumentObject; }
 	
 protected:
 	void RebuildBreadcrumbTrail() const;
diff --git a/Engine/Plugins/Experimental/Workspace/Source/WorkspaceEditor/Private/WorkspaceEditor.cpp b/Engine/Plugins/Experimental/Workspace/Source/WorkspaceEditor/Private/WorkspaceEditor.cpp
--- a/Engine/Plugins/Experimental/Workspace/Source/WorkspaceEditor/Private/WorkspaceEditor.cpp
+++ b/Engine/Plugins/Experimental/Workspace/Source/WorkspaceEditor/Private/WorkspaceEditor.cpp
@@ -19,6 +19,9 @@
 #include "Dialogs/Dialogs.h"
 #include "Framework/Docking/LayoutExtender.h"
 #include "ToolMenus.h"
+#include "AssetEditorModeManager.h"
+#include "Selection.h"
+#include "EditorModeManager.h"
 
 #define LOCTEXT_NAMESPACE "WorkspaceEditor"
 
@@ -216,6 +219,12 @@ void FWorkspaceEditor::SaveEditedObjectState() const
 TSharedPtr<SDockTab> FWorkspaceEditor::OpenDocument(const UObject* InForObject, FDocumentTracker::EOpenDocumentCause InCause)
 {
 	const TSharedRef<FTabPayload_UObject> Payload = FTabPayload_UObject::Make(InForObject);
+	const bool bIsSupportedDocument = DocumentManager->FindSupportingFactory(Payload).IsValid();
+	if (bIsSupportedDocument)
+	{
+		AddEditingObject(const_cast<UObject*>(InForObject));
+	}
+	
 	TSharedPtr<SDockTab> NewTab = DocumentManager->OpenDocument(Payload, InCause);
 
 	if(InCause != FDocumentTracker::RestorePreviousDocument)
@@ -299,6 +308,28 @@ void FWorkspaceEditor::SetGlobalSelection(FGlobalSelectionId SelectionId, FOnCle
 	LastOnClearSelectionDelegate = OnClearSelectionDelegate;
 }
 
+void FWorkspaceEditor::SetFocussedAsset(const TObjectPtr<UObject> InAsset)
+{
+	FAssetEditorModeManager* ModeManager = static_cast<FAssetEditorModeManager*>(&GetEditorModeManager());	
+	ModeManager->GetSelectedObjects()->DeselectAll();
+	
+	if (InAsset != nullptr)
+	{
+		ModeManager->GetSelectedObjects()->Select(InAsset);
+	}
+}
+
+const TObjectPtr<UObject> FWorkspaceEditor::GetFocussedAssetOfClass(const TObjectPtr<UClass> AssetClass) const
+{
+	FAssetEditorModeManager* ModeManager = static_cast<FAssetEditorModeManager*>(&GetEditorModeManager());	
+	if (USelection* Selection = ModeManager->GetSelectedObjects())
+	{
+		return Selection->GetTop(AssetClass);
+	}
+
+	return nullptr;
+}
+
 void FWorkspaceEditor::BindCommands()
 {	
 	FWorkspaceAssetEditorCommands::Register();
@@ -473,6 +504,28 @@ FText FWorkspaceEditor::GetTabSuffix() const
 	return CanSaveAsset() ? LOCTEXT("TabSuffixAsterix", "*") : FText::GetEmpty();
 }
 
+FText FWorkspaceEditor::GetToolkitName() const
+{
+	UObject* const* WorkspaceObject = GetEditingObjects().FindByPredicate([](const UObject* Object) { return Object && Object->IsA<UWorkspace>(); });
+	check (WorkspaceObject != nullptr);
+	
+	return GetLabelForObject(*WorkspaceObject);
+}
+
+FText FWorkspaceEditor::GetToolkitToolTipText() const
+{
+	UObject* const* WorkspaceObject = GetEditingObjects().FindByPredicate([](const UObject* Object) { return Object && Object->IsA<UWorkspace>(); });
+	check (WorkspaceObject != nullptr);
+	
+	FText FocussedAssetText = FText::FromName(NAME_None);	
+	if (const UObject* FocussedObject = GetFocussedAsset())
+	{
+		 FocussedAssetText = GetLabelForObject(FocussedObject);
+	}
+	
+	return FText::Format(LOCTEXT("TookitTooltipFormat", "{0} ({1})"), GetToolTipTextForObject(*WorkspaceObject), FocussedAssetText);
+}
+
 void FWorkspaceEditor::RecordDocumentState(const TInstancedStruct<FWorkspaceDocumentState>& InState) const
 {
 	UWorkspaceState* State = Workspace->GetState();
@@ -573,6 +626,15 @@ bool FWorkspaceEditor::ShouldReopenEditorForSavedAsset(const UObject* Asset) con
 {
 	return !bClosingDown;
 }
+
+void FWorkspaceEditor::RemoveEditingObject(UObject* Object)
+{
+	IWorkspaceEditor::RemoveEditingObject(Object);	
+	if (GetFocussedAsset() == Object)
+	{
+		SetFocussedAsset(nullptr);
+	}
+}
 }
 
 #undef LOCTEXT_NAMESPACE
\ No newline at end of file
diff --git a/Engine/Plugins/Experimental/Workspace/Source/WorkspaceEditor/Private/WorkspaceEditor.h b/Engine/Plugins/Experimental/Workspace/Source/WorkspaceEditor/Private/WorkspaceEditor.h
--- a/Engine/Plugins/Experimental/Workspace/Source/WorkspaceEditor/Private/WorkspaceEditor.h
+++ b/Engine/Plugins/Experimental/Workspace/Source/WorkspaceEditor/Private/WorkspaceEditor.h
@@ -68,11 +68,14 @@ private:
 	virtual void OnClose() override;
 	virtual void RegisterToolbar() override;
 	virtual bool ShouldReopenEditorForSavedAsset(const UObject* Asset) const override;
+	virtual void RemoveEditingObject(UObject* Object);
 
 	// FAssetEditorToolkit interface
 	virtual void GetSaveableObjects(TArray<UObject*>& OutObjects) const override;
 	virtual bool CanSaveAsset() const override;
 	virtual FText GetTabSuffix() const override;
+	virtual FText GetToolkitName() const override;
+	virtual FText GetToolkitToolTipText() const override;	
 	virtual bool IsFindInContentBrowserButtonVisible() const { return false; }
 
 	// FGCObject interface
@@ -92,8 +95,9 @@ private:
 	virtual void SetDetailsObjects(const TArray<UObject*>& InObjects) override;
 	virtual void RefreshDetails() override;
 	virtual UWorkspaceSchema* GetSchema() const override;
-
 	virtual void SetGlobalSelection(FGlobalSelectionId SelectionId, FOnClearGlobalSelection OnClearSelectionDelegate) override;
+	virtual void SetFocussedAsset(const TObjectPtr<UObject> InAsset) override;	
+	virtual const TObjectPtr<UObject> GetFocussedAssetOfClass(const TObjectPtr<UClass> AssetClass) const override;	
 
 	void BindCommands();
 
diff --git a/Engine/Plugins/Experimental/Workspace/Source/WorkspaceEditor/Public/IWorkspaceEditor.h b/Engine/Plugins/Experimental/Workspace/Source/WorkspaceEditor/Public/IWorkspaceEditor.h
--- a/Engine/Plugins/Experimental/Workspace/Source/WorkspaceEditor/Public/IWorkspaceEditor.h
+++ b/Engine/Plugins/Experimental/Workspace/Source/WorkspaceEditor/Public/IWorkspaceEditor.h
@@ -38,6 +38,13 @@ public:
 
 	// Set the _current_ global selection (last SWidget with selection set) with delegate to clear it selection on next SetGlobalSelection()
 	virtual void SetGlobalSelection(FGlobalSelectionId SelectionId, FOnClearGlobalSelection OnClearSelectionDelegate) = 0;
+
+	virtual void SetFocussedAsset(const TObjectPtr<UObject> InAsset) = 0;
+	virtual const TObjectPtr<UObject> GetFocussedAssetOfClass(const TObjectPtr<UClass> InClass ) const = 0;
+	
+	template<typename AssetClass>
+	TObjectPtr<AssetClass> GetFocussedAsset() const { return GetFocussedAssetOfClass(AssetClass::StaticClass()); }
+	TObjectPtr<UObject> GetFocussedAsset() const { return GetFocussedAssetOfClass(UObject::StaticClass()); }
 };
 
 }
\ No newline at end of file
