Anim Next:
-Added browse to asset for workspace and outliner
-Fix delete and undo/redo for context menu delete in outline

diff --git a/Engine/Plugins/Experimental/Workspace/Source/WorkspaceEditor/Private/Outliner/WorkspaceOutlinerMode.cpp b/Engine/Plugins/Experimental/Workspace/Source/WorkspaceEditor/Private/Outliner/WorkspaceOutlinerMode.cpp
--- a/Engine/Plugins/Experimental/Workspace/Source/WorkspaceEditor/Private/Outliner/WorkspaceOutlinerMode.cpp
+++ b/Engine/Plugins/Experimental/Workspace/Source/WorkspaceEditor/Private/Outliner/WorkspaceOutlinerMode.cpp
@@ -116,10 +116,11 @@ TSharedPtr<SWidget> FWorkspaceOutlinerMode::CreateContextMenu()
 										}
 									}
 
-									if (EditingObjects.Num() == 1)
+									if (EditingObjects.Num() > 0)
 									{
 										if (UWorkspace* Workspace = Cast<UWorkspace>(EditingObjects[0]))
 										{
+											FScopedTransaction Transaction(LOCTEXT("RemoveAssets", "Remove assets from workspace"));
 											for (const FSoftObjectPath& AssetPath : AssetPaths)
 											{
 												Workspace->RemoveAsset(AssetPath.TryLoad());
@@ -129,6 +130,31 @@ TSharedPtr<SWidget> FWorkspaceOutlinerMode::CreateContextMenu()
 								}
 							}))
 						);
+
+						AssetsSection.AddMenuEntry(TEXT("BrowseToAsset"),
+							LOCTEXT("BrowseToAssetLabel", "Browse to Asset"),
+							LOCTEXT("BrowseToAssetTooltip", "Browse to the selected assets in the content browser"),
+							FSlateIcon(FAppStyle::Get().GetStyleSetName(), "SystemWideCommands.FindInContentBrowser.Small"),
+							FUIAction(FExecuteAction::CreateLambda([SelectedExports=MenuContext->SelectedExports]()
+								{
+									const IAssetRegistry& AssetRegistry = FAssetRegistryModule::GetRegistry();
+									FARFilter Filter;	
+									for (const FWorkspaceOutlinerItemExport& ItemExport : SelectedExports)
+									{
+										if (ItemExport.GetParentIdentifier() == NAME_None)
+										{
+											Filter.SoftObjectPaths.AddUnique(ItemExport.GetAssetPath());
+										}
+									}
+
+									TArray<FAssetData> AssetDataList;
+									AssetRegistry.GetAssets(Filter, AssetDataList);
+									if (AssetDataList.IsEmpty() == false)
+									{
+										GEditor->SyncBrowserToObjects(AssetDataList);
+									}
+								}))
+						);
 					}
 
 					auto IsPackageDirty = [](const UPackage* Package)-> bool
diff --git a/Engine/Plugins/Experimental/Workspace/Source/WorkspaceEditor/Private/SWorkspaceTabWrapper.cpp b/Engine/Plugins/Experimental/Workspace/Source/WorkspaceEditor/Private/SWorkspaceTabWrapper.cpp
--- a/Engine/Plugins/Experimental/Workspace/Source/WorkspaceEditor/Private/SWorkspaceTabWrapper.cpp
+++ b/Engine/Plugins/Experimental/Workspace/Source/WorkspaceEditor/Private/SWorkspaceTabWrapper.cpp
@@ -49,6 +49,21 @@ void SWorkspaceTabWrapper::Construct( const FArguments& InArgs, TSharedPtr<class
 		FToolMenuSection& Section = ToolBarMenu->AddSection("AssetActions");
 		Section.AddMenuEntry("Save", TAttribute<FText>(), LOCTEXT("SaveButtonTooltip", "Save Asset"), FSlateIcon(FAppStyle::Get().GetStyleSetName(), "AssetEditor.SaveAsset"),
 			FUIAction(FExecuteAction::CreateSP(this, &SWorkspaceTabWrapper::ExecuteSave), FCanExecuteAction::CreateSP(this, &SWorkspaceTabWrapper::CanExecuteSave), FIsActionChecked(), FIsActionButtonVisible::CreateSP(this, &SWorkspaceTabWrapper::IsSaveButtonVisible)));
+	
+		const FToolMenuEntry FindInContentBrowserEntry = FToolMenuEntry::InitToolBarButton(
+			"FindInContentBrowser",
+			FUIAction(FExecuteAction::CreateSPLambda(this, [this]() 
+				{
+					if (UObject* Asset = WeakDocumentObject.Get())
+					{
+						GEditor->SyncBrowserToObject(Asset);
+					}
+				})),
+			FText::GetEmpty(),
+			LOCTEXT("FindInContentBrowserTooltip", "Finds this asset in the content browser"),
+			FSlateIcon(FAppStyle::Get().GetStyleSetName(), "SystemWideCommands.FindInContentBrowser")
+		);
+		Section.AddEntry(FindInContentBrowserEntry);
 	}
 	
 	ChildSlot
diff --git a/Engine/Plugins/Experimental/Workspace/Source/WorkspaceEditor/Private/WorkspaceEditor.cpp b/Engine/Plugins/Experimental/Workspace/Source/WorkspaceEditor/Private/WorkspaceEditor.cpp
--- a/Engine/Plugins/Experimental/Workspace/Source/WorkspaceEditor/Private/WorkspaceEditor.cpp
+++ b/Engine/Plugins/Experimental/Workspace/Source/WorkspaceEditor/Private/WorkspaceEditor.cpp
@@ -567,6 +567,14 @@ bool FWorkspaceEditor::AreAssetEntriesModified() const
 	return false;
 }
 
+void FWorkspaceEditor::FindInContentBrowser_Execute()
+{
+	if (IsValid(Workspace))
+	{
+		GEditor->SyncBrowserToObject(Workspace);
+	}
+}
+
 bool FWorkspaceEditor::OnRequestClose(EAssetEditorCloseReason InCloseReason)
 {
 	TGuardValue<bool> ClosingDown(bClosingDown, true);
diff --git a/Engine/Plugins/Experimental/Workspace/Source/WorkspaceEditor/Private/WorkspaceEditor.h b/Engine/Plugins/Experimental/Workspace/Source/WorkspaceEditor/Private/WorkspaceEditor.h
--- a/Engine/Plugins/Experimental/Workspace/Source/WorkspaceEditor/Private/WorkspaceEditor.h
+++ b/Engine/Plugins/Experimental/Workspace/Source/WorkspaceEditor/Private/WorkspaceEditor.h
@@ -76,7 +76,9 @@ private:
 	virtual FText GetTabSuffix() const override;
 	virtual FText GetToolkitName() const override;
 	virtual FText GetToolkitToolTipText() const override;	
-	virtual bool IsFindInContentBrowserButtonVisible() const { return false; }
+	virtual bool IsFindInContentBrowserButtonVisible() const override { return true; }
+	virtual void FindInContentBrowser_Execute() override;
+
 
 	// FGCObject interface
 	virtual void AddReferencedObjects(FReferenceCollector& Collector) override
