Fixed lambda parameter shadowing member parameter causing Horde CLang error

#rnx

diff --git a/Engine/Plugins/Experimental/Workspace/Source/WorkspaceEditor/Private/SGraphDocument.cpp b/Engine/Plugins/Experimental/Workspace/Source/WorkspaceEditor/Private/SGraphDocument.cpp
--- a/Engine/Plugins/Experimental/Workspace/Source/WorkspaceEditor/Private/SGraphDocument.cpp
+++ b/Engine/Plugins/Experimental/Workspace/Source/WorkspaceEditor/Private/SGraphDocument.cpp
@@ -27,7 +27,6 @@ void SGraphDocument::Construct(const FArguments& InArgs, TSharedRef<FWorkspaceEd
 	OnPasteNodes = InArgs._OnPasteNodes;
 	OnCanDuplicateSelectedNodes = InArgs._OnCanDuplicateSelectedNodes;
 	OnDuplicateSelectedNodes = InArgs._OnDuplicateSelectedNodes;
-	OnNodeDoubleClicked = InArgs._OnNodeDoubleClicked;
 
 	SGraphEditor::FGraphEditorEvents Events;
 	Events.OnCreateActionMenu = SGraphEditor::FOnCreateActionMenu::CreateLambda([this, OnCreateActionMenu = InArgs._OnCreateActionMenu](UEdGraph* InGraph, const FVector2D& InNodePosition, const TArray<UEdGraphPin*>& InDraggedPins, bool bInAutoExpand, SGraphEditor::FActionMenuClosed InOnMenuClosed)
diff --git a/Engine/Plugins/Experimental/Workspace/Source/WorkspaceEditor/Private/SGraphDocument.h b/Engine/Plugins/Experimental/Workspace/Source/WorkspaceEditor/Private/SGraphDocument.h
--- a/Engine/Plugins/Experimental/Workspace/Source/WorkspaceEditor/Private/SGraphDocument.h
+++ b/Engine/Plugins/Experimental/Workspace/Source/WorkspaceEditor/Private/SGraphDocument.h
@@ -107,9 +107,6 @@ class SGraphDocument : public SCompoundWidget
 	FOnCanPerformActionOnSelectedNodes OnCanDuplicateSelectedNodes;
 	FOnDuplicateSelectedNodes OnDuplicateSelectedNodes;
 
-	// Node Double click
-	FOnNodeDoubleClicked OnNodeDoubleClicked;
-
 	friend class FWorkspaceEditorModule;
 };
 
