Introduces workspace level concept (v0.1) of asset references. Any export made with FWorkspaceOutlinerAssetReferenceItemData is treated as a direct references to an on-disk asset and will be reflected in the Outliner accordingly. (recursive references are caught and reflected through coloration).
	- Added FWorkspaceOutlinerAssetReferenceItemData and its respective item details
	- Added GetObjectReferences to FAnimNextStateTreeTaskBase - used to export asset reference data to asset-registry
	- Added IsExpandedByDefault to IWorkspaceOutlinerItemDetails - allows for setting expansion off by default for asset references in outliner
	- Added different ways of retrieving UObjects and asset paths from FWorkspaceOutlinerItemExport - to be able to deal with asset references more easily
	- FUtils::CreateSubGraphsOutlinerItemsRecursive now exports AnimNextGraph asset references from SubGrap traits, and RunAnimNextGraph nodes

diff --git a/Engine/Plugins/Experimental/AnimNext/Source/AnimNextEditor/Private/Common/AnimNextCollapseNodeItemDetails.cpp b/Engine/Plugins/Experimental/AnimNext/Source/AnimNextEditor/Private/Common/AnimNextCollapseNodeItemDetails.cpp
--- a/Engine/Plugins/Experimental/AnimNext/Source/AnimNextEditor/Private/Common/AnimNextCollapseNodeItemDetails.cpp
+++ b/Engine/Plugins/Experimental/AnimNext/Source/AnimNextEditor/Private/Common/AnimNextCollapseNodeItemDetails.cpp
@@ -20,7 +20,7 @@
 namespace UE::AnimNext::Editor
 {
 
-void FAnimNextCollapseNodeItemDetails::HandleDoubleClick(const FToolMenuContext& ToolMenuContext) const
+bool FAnimNextCollapseNodeItemDetails::HandleDoubleClick(const FToolMenuContext& ToolMenuContext) const
 {
 	const UWorkspaceItemMenuContext* WorkspaceItemContext = ToolMenuContext.FindContext<UWorkspaceItemMenuContext>();
 	const UAssetEditorToolkitMenuContext* AssetEditorContext = ToolMenuContext.FindContext<UAssetEditorToolkitMenuContext>();
@@ -36,10 +36,13 @@ void FAnimNextCollapseNodeItemDetails::HandleDoubleClick(const FToolMenuContext&
 				{
 					UObject* EditorObject = CollapseGraphData.EditorObject.Get();
 					WorkspaceEditor->OpenObjects({ EditorObject });
+					return true;
 				}
 			}
 		}
 	}
+	
+	return false;
 }
 
 bool FAnimNextCollapseNodeItemDetails::CanDelete(const FWorkspaceOutlinerItemExport& Export) const
diff --git a/Engine/Plugins/Experimental/AnimNext/Source/AnimNextEditor/Private/Common/AnimNextCollapseNodeItemDetails.h b/Engine/Plugins/Experimental/AnimNext/Source/AnimNextEditor/Private/Common/AnimNextCollapseNodeItemDetails.h
--- a/Engine/Plugins/Experimental/AnimNext/Source/AnimNextEditor/Private/Common/AnimNextCollapseNodeItemDetails.h
+++ b/Engine/Plugins/Experimental/AnimNext/Source/AnimNextEditor/Private/Common/AnimNextCollapseNodeItemDetails.h
@@ -13,7 +13,7 @@ public:
 	FAnimNextCollapseNodeItemDetails() = default;
 	virtual ~FAnimNextCollapseNodeItemDetails() override = default;
 
-	virtual void HandleDoubleClick(const FToolMenuContext& ToolMenuContext) const override;
+	virtual bool HandleDoubleClick(const FToolMenuContext& ToolMenuContext) const override;
 	virtual bool CanDelete(const FWorkspaceOutlinerItemExport& Export) const override;
 	virtual void Delete(TConstArrayView<FWorkspaceOutlinerItemExport> Exports) const override;
 	virtual bool CanRename(const FWorkspaceOutlinerItemExport& Export) const override;
diff --git a/Engine/Plugins/Experimental/AnimNext/Source/AnimNextEditor/Private/Common/AnimNextFunctionItemDetails.cpp b/Engine/Plugins/Experimental/AnimNext/Source/AnimNextEditor/Private/Common/AnimNextFunctionItemDetails.cpp
--- a/Engine/Plugins/Experimental/AnimNext/Source/AnimNextEditor/Private/Common/AnimNextFunctionItemDetails.cpp
+++ b/Engine/Plugins/Experimental/AnimNext/Source/AnimNextEditor/Private/Common/AnimNextFunctionItemDetails.cpp
@@ -20,7 +20,7 @@
 namespace UE::AnimNext::Editor
 {
 
-void FAnimNextFunctionItemDetails::HandleDoubleClick(const FToolMenuContext& ToolMenuContext) const
+bool FAnimNextFunctionItemDetails::HandleDoubleClick(const FToolMenuContext& ToolMenuContext) const
 {
 	const UWorkspaceItemMenuContext* WorkspaceItemContext = ToolMenuContext.FindContext<UWorkspaceItemMenuContext>();
 	const UAssetEditorToolkitMenuContext* AssetEditorContext = ToolMenuContext.FindContext<UAssetEditorToolkitMenuContext>();
@@ -35,10 +35,13 @@ void FAnimNextFunctionItemDetails::HandleDoubleClick(const FToolMenuContext& Too
 				if (GraphFunctionData.EditorObject.IsValid())
 				{
 					WorkspaceEditor->OpenObjects({ GraphFunctionData.EditorObject.Get() });
+					return true;					
 				}
 			}
 		}
 	}
+
+	return false;
 }
 
 bool FAnimNextFunctionItemDetails::CanDelete(const FWorkspaceOutlinerItemExport& Export) const
diff --git a/Engine/Plugins/Experimental/AnimNext/Source/AnimNextEditor/Private/Common/AnimNextFunctionItemDetails.h b/Engine/Plugins/Experimental/AnimNext/Source/AnimNextEditor/Private/Common/AnimNextFunctionItemDetails.h
--- a/Engine/Plugins/Experimental/AnimNext/Source/AnimNextEditor/Private/Common/AnimNextFunctionItemDetails.h
+++ b/Engine/Plugins/Experimental/AnimNext/Source/AnimNextEditor/Private/Common/AnimNextFunctionItemDetails.h
@@ -13,7 +13,7 @@ public:
 	FAnimNextFunctionItemDetails() = default;
 	virtual ~FAnimNextFunctionItemDetails() override = default;
 
-	virtual void HandleDoubleClick(const FToolMenuContext& ToolMenuContext) const override;
+	virtual bool HandleDoubleClick(const FToolMenuContext& ToolMenuContext) const override;
 	virtual bool CanDelete(const FWorkspaceOutlinerItemExport& Export) const override;
 	virtual void Delete(TConstArrayView<FWorkspaceOutlinerItemExport> Exports) const override;
 	virtual bool CanRename(const FWorkspaceOutlinerItemExport& Export) const override;
diff --git a/Engine/Plugins/Experimental/AnimNext/Source/AnimNextEditor/Private/Common/AnimNextGraphItemDetails.cpp b/Engine/Plugins/Experimental/AnimNext/Source/AnimNextEditor/Private/Common/AnimNextGraphItemDetails.cpp
--- a/Engine/Plugins/Experimental/AnimNext/Source/AnimNextEditor/Private/Common/AnimNextGraphItemDetails.cpp
+++ b/Engine/Plugins/Experimental/AnimNext/Source/AnimNextEditor/Private/Common/AnimNextGraphItemDetails.cpp
@@ -20,7 +20,7 @@
 namespace UE::AnimNext::Editor
 {
 
-void FAnimNextGraphItemDetails::HandleDoubleClick(const FToolMenuContext& ToolMenuContext) const
+bool FAnimNextGraphItemDetails::HandleDoubleClick(const FToolMenuContext& ToolMenuContext) const
 {
 	const UWorkspaceItemMenuContext* WorkspaceItemContext = ToolMenuContext.FindContext<UWorkspaceItemMenuContext>();
 	const UAssetEditorToolkitMenuContext* AssetEditorContext = ToolMenuContext.FindContext<UAssetEditorToolkitMenuContext>(); 
@@ -41,6 +41,7 @@ void FAnimNextGraphItemDetails::HandleDoubleClick(const FToolMenuContext& ToolMe
 							if(UObject* EditorObject = RigVMClientHost->GetEditorObjectForRigVMGraph(RigVMGraph))
 							{
 								WorkspaceEditor->OpenObjects({EditorObject});
+								return true;
 							}
 						}
 					}
@@ -48,6 +49,8 @@ void FAnimNextGraphItemDetails::HandleDoubleClick(const FToolMenuContext& ToolMe
 			}
 		}
 	}
+
+	return false;
 }
 
 bool FAnimNextGraphItemDetails::CanDelete(const FWorkspaceOutlinerItemExport& Export) const
diff --git a/Engine/Plugins/Experimental/AnimNext/Source/AnimNextEditor/Private/Common/AnimNextGraphItemDetails.h b/Engine/Plugins/Experimental/AnimNext/Source/AnimNextEditor/Private/Common/AnimNextGraphItemDetails.h
--- a/Engine/Plugins/Experimental/AnimNext/Source/AnimNextEditor/Private/Common/AnimNextGraphItemDetails.h
+++ b/Engine/Plugins/Experimental/AnimNext/Source/AnimNextEditor/Private/Common/AnimNextGraphItemDetails.h
@@ -12,7 +12,7 @@ class FAnimNextGraphItemDetails : public UE::Workspace::IWorkspaceOutlinerItemDe
 public:
 	FAnimNextGraphItemDetails() = default;
 	virtual ~FAnimNextGraphItemDetails() override = default;
-	virtual void HandleDoubleClick(const FToolMenuContext& ToolMenuContext) const override;
+	virtual bool HandleDoubleClick(const FToolMenuContext& ToolMenuContext) const override;
 	virtual bool CanDelete(const FWorkspaceOutlinerItemExport& Export) const override;
 	virtual void Delete(TConstArrayView<FWorkspaceOutlinerItemExport> Exports) const override; 
 	virtual bool CanRename(const FWorkspaceOutlinerItemExport& Export) const override;
diff --git a/Engine/Plugins/Experimental/AnimNext/Source/AnimNextEditor/Private/Variables/SVariablesView.cpp b/Engine/Plugins/Experimental/AnimNext/Source/AnimNextEditor/Private/Variables/SVariablesView.cpp
--- a/Engine/Plugins/Experimental/AnimNext/Source/AnimNextEditor/Private/Variables/SVariablesView.cpp
+++ b/Engine/Plugins/Experimental/AnimNext/Source/AnimNextEditor/Private/Variables/SVariablesView.cpp
@@ -249,16 +249,19 @@ void SVariablesView::HandleWorkspaceOutlinerSelectionChanged(TConstArrayView<FWo
 	TArray<FSoftObjectPath> AssetsToAsyncLoad;
 	for(const FWorkspaceOutlinerItemExport& Export : InExports)
 	{
-		if(Export.GetAssetPath().IsAsset() && Export.GetData().GetScriptStruct()->IsChildOf(FAnimNextRigVMAssetOutlinerData::StaticStruct()))
+		const FSoftObjectPath FirstAssetPath = Export.GetFirstAssetPath();
+		if(FirstAssetPath.IsAsset())
 		{
-			Assets.Add(TSoftObjectPtr<UAnimNextRigVMAsset>(Export.GetAssetPath()));
-			UAnimNextRigVMAsset* ResolvedAsset = Cast<UAnimNextRigVMAsset>(Export.GetAssetPath().ResolveObject());
-			if(ResolvedAsset == nullptr)
+			Assets.Add(TSoftObjectPtr<UAnimNextRigVMAsset>(FirstAssetPath));
+
+			UObject* ResolvedObject = FirstAssetPath.ResolveObject();
+			UAnimNextRigVMAsset* ResolvedAsset = Cast<UAnimNextRigVMAsset>(ResolvedObject);
+			if(ResolvedObject == nullptr)
 			{
-				AssetsToAsyncLoad.Add(Export.GetAssetPath());
+				AssetsToAsyncLoad.Add(FirstAssetPath);
 			}
 		}
-		else if(Export.GetData().GetScriptStruct()->IsChildOf(FAnimNextAssetEntryOutlinerData::StaticStruct()))
+		else if(Export.HasData() && Export.GetData().GetScriptStruct()->IsChildOf(FAnimNextAssetEntryOutlinerData::StaticStruct()))
 		{
 			// TODO: As we are not showing references in the workspace yet, we are just traversing the outer here and showing the containing asset's
 			// variables. Eventually we would want to recurse up to the root of the export hierarchy, which will entail changing the signature of
diff --git a/Engine/Plugins/Experimental/AnimNext/Source/AnimNextUncookedOnly/Internal/Entries/AnimNextVariableEntry.h b/Engine/Plugins/Experimental/AnimNext/Source/AnimNextUncookedOnly/Internal/Entries/AnimNextVariableEntry.h
--- a/Engine/Plugins/Experimental/AnimNext/Source/AnimNextUncookedOnly/Internal/Entries/AnimNextVariableEntry.h
+++ b/Engine/Plugins/Experimental/AnimNext/Source/AnimNextUncookedOnly/Internal/Entries/AnimNextVariableEntry.h
@@ -54,6 +54,7 @@ class UAnimNextVariableEntry : public UAnimNextRigVMAssetEntry, public IAnimNext
 	friend class UAnimNextModule_EditorData;
 	friend class UAnimNextDataInterface_EditorData;
 	friend class UAnimNextDataInterfaceEntry;
+	friend class UAnimNextStateTreeTreeEditorData;
 	friend class UE::AnimNext::Tests::FEditor_Variables;
 	friend class UE::AnimNext::Tests::FEditor_AnimGraph_Variables;
 	friend class UE::AnimNext::Tests::FDataInterfaceCompile;
diff --git a/Engine/Plugins/Experimental/AnimNext/Source/AnimNextUncookedOnly/Internal/UncookedOnlyUtils.h b/Engine/Plugins/Experimental/AnimNext/Source/AnimNextUncookedOnly/Internal/UncookedOnlyUtils.h
--- a/Engine/Plugins/Experimental/AnimNext/Source/AnimNextUncookedOnly/Internal/UncookedOnlyUtils.h
+++ b/Engine/Plugins/Experimental/AnimNext/Source/AnimNextUncookedOnly/Internal/UncookedOnlyUtils.h
@@ -174,8 +174,8 @@ struct ANIMNEXTUNCOOKEDONLY_API FUtils
 	static void GetAssetVariables(const UAnimNextRigVMAssetEditorData* EditorData, TSet<FAnimNextAssetRegistryExportedVariable>& OutExports);
 
 	// Gets the asset-registry information needed for representing the contained data into the Workspace Outliner
-	static void GetAssetOutlinerItems(const UAnimNextRigVMAssetEditorData* EditorData, FWorkspaceOutlinerItemExports& OutExports);
-	static void CreateSubGraphsOutlinerItemsRecursive(const UAnimNextRigVMAssetEditorData* EditorData, FWorkspaceOutlinerItemExports& OutExports, FWorkspaceOutlinerItemExport& ParentExport, URigVMEdGraph* RigVMEdGraph);
+	static void GetAssetOutlinerItems(const UAnimNextRigVMAssetEditorData* EditorData, FWorkspaceOutlinerItemExports& OutExports, FAssetRegistryTagsContext Context);
+	static void CreateSubGraphsOutlinerItemsRecursive(const UAnimNextRigVMAssetEditorData* EditorData, FWorkspaceOutlinerItemExports& OutExports, FWorkspaceOutlinerItemExport& ParentExport, URigVMEdGraph* RigVMEdGraph, FAssetRegistryTagsContext Context);
 	static void CreateFunctionLibraryOutlinerItemsRecursive(const UAnimNextRigVMAssetEditorData* EditorData, FWorkspaceOutlinerItemExports& OutExports, FWorkspaceOutlinerItemExport& ParentExport, const TArray<FRigVMGraphFunctionData>& PublicFunctions, const TArray<FRigVMGraphFunctionData>& PrivateFunctions);
 	static void CreateFunctionsOutlinerItemsRecursive(const UAnimNextRigVMAssetEditorData* EditorData, FWorkspaceOutlinerItemExports& OutExports, FWorkspaceOutlinerItemExport& ParentExport, const TArray<FRigVMGraphFunctionData>& Functions, bool bPublicFunctions);
 
diff --git a/Engine/Plugins/Experimental/AnimNext/Source/AnimNextUncookedOnly/Private/AnimNextAssetWorkspaceAssetUserData.cpp b/Engine/Plugins/Experimental/AnimNext/Source/AnimNextUncookedOnly/Private/AnimNextAssetWorkspaceAssetUserData.cpp
--- a/Engine/Plugins/Experimental/AnimNext/Source/AnimNextUncookedOnly/Private/AnimNextAssetWorkspaceAssetUserData.cpp
+++ b/Engine/Plugins/Experimental/AnimNext/Source/AnimNextUncookedOnly/Private/AnimNextAssetWorkspaceAssetUserData.cpp
@@ -35,7 +35,7 @@ void UAnimNextAssetWorkspaceAssetUserData::GetAssetRegistryTags(FAssetRegistryTa
 			Data.Asset = Asset;
 		}
 	
-		UE::AnimNext::UncookedOnly::FUtils::GetAssetOutlinerItems(GraphEditorData, Exports);
+		UE::AnimNext::UncookedOnly::FUtils::GetAssetOutlinerItems(GraphEditorData, Exports, Context);
 	}
 	
 	FString TagValue;
diff --git a/Engine/Plugins/Experimental/AnimNext/Source/AnimNextUncookedOnly/Private/UncookedOnlyUtils.cpp b/Engine/Plugins/Experimental/AnimNext/Source/AnimNextUncookedOnly/Private/UncookedOnlyUtils.cpp
--- a/Engine/Plugins/Experimental/AnimNext/Source/AnimNextUncookedOnly/Private/UncookedOnlyUtils.cpp
+++ b/Engine/Plugins/Experimental/AnimNext/Source/AnimNextUncookedOnly/Private/UncookedOnlyUtils.cpp
@@ -27,6 +27,7 @@
 #include "DataInterface/AnimNextDataInterface.h"
 #include "DataInterface/AnimNextDataInterface_EditorData.h"
 #include "Logging/MessageLog.h"
+#include "UObject/AssetRegistryTagsContext.h"
 
 #define LOCTEXT_NAMESPACE "AnimNextUncookedOnlyUtils"
 
@@ -810,7 +811,7 @@ void FUtils::GetAssetVariables(const UAnimNextRigVMAssetEditorData* InEditorData
 	}
 }
 
-void FUtils::GetAssetOutlinerItems(const UAnimNextRigVMAssetEditorData* EditorData, FWorkspaceOutlinerItemExports& OutExports)
+void FUtils::GetAssetOutlinerItems(const UAnimNextRigVMAssetEditorData* EditorData, FWorkspaceOutlinerItemExports& OutExports, FAssetRegistryTagsContext Context)
 {
 	FWorkspaceOutlinerItemExport AssetIdentifier = FWorkspaceOutlinerItemExport(EditorData->GetOuter()->GetFName(), EditorData->GetOuter());
 	for(UAnimNextRigVMAssetEntry* Entry : EditorData->Entries)
@@ -821,7 +822,7 @@ void FUtils::GetAssetOutlinerItems(const UAnimNextRigVMAssetEditorData* EditorDa
 			{
 				if (URigVMEdGraph* RigVMEdGraph = GraphInterface->GetEdGraph())
 				{
-					CreateSubGraphsOutlinerItemsRecursive(EditorData, OutExports, AssetIdentifier, RigVMEdGraph);
+					CreateSubGraphsOutlinerItemsRecursive(EditorData, OutExports, AssetIdentifier, RigVMEdGraph, Context);
 				}
 			}
 			else
@@ -835,7 +836,7 @@ void FUtils::GetAssetOutlinerItems(const UAnimNextRigVMAssetEditorData* EditorDa
 
 				if (URigVMEdGraph* RigVMEdGraph = GraphInterface->GetEdGraph())
 				{
-					CreateSubGraphsOutlinerItemsRecursive(EditorData, OutExports, Export, RigVMEdGraph);
+					CreateSubGraphsOutlinerItemsRecursive(EditorData, OutExports, Export, RigVMEdGraph, Context);
 				}
 			}
 		}
@@ -844,13 +845,74 @@ void FUtils::GetAssetOutlinerItems(const UAnimNextRigVMAssetEditorData* EditorDa
 	//CreateFunctionLibraryOutlinerItemsRecursive(EditorData, OutExports, AssetIdentifier, EditorData->GetRigVMGraphFunctionStore()->PublicFunctions, EditorData->GetRigVMGraphFunctionStore()->PrivateFunctions);
 }
 
-void FUtils::CreateSubGraphsOutlinerItemsRecursive(const UAnimNextRigVMAssetEditorData* EditorData, FWorkspaceOutlinerItemExports& OutExports, FWorkspaceOutlinerItemExport& ParentExport, URigVMEdGraph* RigVMEdGraph)
+void ProcessPinAssetReferences(const URigVMPin* InPin, FWorkspaceOutlinerItemExports& OutExports, const FWorkspaceOutlinerItemExport& ParentExport)
+{
+	if (InPin)
+	{
+		const UObject* TypeObject = InPin->GetCPPTypeObject();
+		if (Cast<UClass>(TypeObject))
+		{
+			const FSoftObjectPath ObjectPath(InPin->GetDefaultValue());
+			if (ObjectPath.IsValid())
+			{
+				FWorkspaceOutlinerItemExport& ReferenceExport = OutExports.Exports.Add_GetRef(FWorkspaceOutlinerItemExport(FName(ObjectPath.ToString()), ParentExport));
+				ReferenceExport.GetData().InitializeAs<FWorkspaceOutlinerAssetReferenceItemData>();
+				ReferenceExport.GetData().GetMutable<FWorkspaceOutlinerAssetReferenceItemData>().ReferredObjectPath = ObjectPath;
+			}
+		}
+
+		for (const URigVMPin* SubPin : InPin->GetSubPins())
+		{
+			ProcessPinAssetReferences(SubPin, OutExports, ParentExport);
+		}
+	}
+	
+}
+
+void FUtils::CreateSubGraphsOutlinerItemsRecursive(const UAnimNextRigVMAssetEditorData* EditorData, FWorkspaceOutlinerItemExports& OutExports, FWorkspaceOutlinerItemExport& ParentExport, URigVMEdGraph* RigVMEdGraph, FAssetRegistryTagsContext Context)
 {
 	if (RigVMEdGraph == nullptr)
 	{
 		return;
 	}
 
+	// Handle pin asset references (disabled during save as GetMetaData can cause StaticFindFast calls which is prohibited during save)
+	if (!Context.IsSaving())
+	{
+		for (const TObjectPtr<class UEdGraphNode>& Node : RigVMEdGraph->Nodes)
+		{
+			if (URigVMEdGraphNode* RigVMEdNode = Cast<URigVMEdGraphNode>(Node))
+			{
+				if (URigVMTemplateNode* TemplateRigVMNode = Cast<URigVMTemplateNode>(RigVMEdNode->GetModelNode()))
+				{
+					if (TemplateRigVMNode->GetScriptStruct()->IsChildOf(FRigUnit_AnimNextBase::StaticStruct()))
+					{
+						for (URigVMPin* ModelPin : TemplateRigVMNode->GetPins())
+						{
+							if (ModelPin->GetDirection() == ERigVMPinDirection::Input)
+							{
+								auto HandlePin = [&OutExports, ParentExport](const URigVMPin* InPin)
+								{
+									if (InPin->GetMetaData(TEXT("ExportAsReference")) == TEXT("true"))
+									{
+										ProcessPinAssetReferences(InPin, OutExports, ParentExport);
+									}										
+								};
+								
+								HandlePin(ModelPin);
+								
+								for (URigVMPin* TraitPin : ModelPin->GetSubPins())
+								{									
+									HandlePin(TraitPin);
+								}
+							}
+						}
+					}
+				}
+			}
+		}
+	}
+
 	// ---- Collapsed graphs ---
 	for (const TObjectPtr<UEdGraph>& SubGraph : RigVMEdGraph->SubGraphs)
 	{
@@ -867,7 +929,7 @@ void FUtils::CreateSubGraphsOutlinerItemsRecursive(const UAnimNextRigVMAssetEdit
 				FAnimNextCollapseGraphOutlinerData& FnGraphData = Export.GetData().GetMutable<FAnimNextCollapseGraphOutlinerData>();
 				FnGraphData.EditorObject = EditorObject;
 
-				CreateSubGraphsOutlinerItemsRecursive(EditorData, OutExports, Export, EditorObject);
+				CreateSubGraphsOutlinerItemsRecursive(EditorData, OutExports, Export, EditorObject, Context);
 			}
 		}
 	}
@@ -899,7 +961,7 @@ void FUtils::CreateSubGraphsOutlinerItemsRecursive(const UAnimNextRigVMAssetEdit
 						FnGraphData.EditorObject = ContainedEdGraph;
 						FnGraphData.EdGraphNode = EdNode;
 						
-						CreateSubGraphsOutlinerItemsRecursive(EditorData, OutExports, Export, ContainedEdGraph);
+						CreateSubGraphsOutlinerItemsRecursive(EditorData, OutExports, Export, ContainedEdGraph, Context);
 					}
 				}
 			}
diff --git a/Engine/Plugins/Experimental/AnimNextAnimGraph/Source/AnimNextAnimGraph/Internal/Graph/RigUnit_AnimNextRunAnimationGraph.h b/Engine/Plugins/Experimental/AnimNextAnimGraph/Source/AnimNextAnimGraph/Internal/Graph/RigUnit_AnimNextRunAnimationGraph.h
--- a/Engine/Plugins/Experimental/AnimNextAnimGraph/Source/AnimNextAnimGraph/Internal/Graph/RigUnit_AnimNextRunAnimationGraph.h
+++ b/Engine/Plugins/Experimental/AnimNextAnimGraph/Source/AnimNextAnimGraph/Internal/Graph/RigUnit_AnimNextRunAnimationGraph.h
@@ -19,7 +19,7 @@ struct ANIMNEXTANIMGRAPH_API FRigUnit_AnimNextRunAnimationGraph : public FRigUni
 	void Execute();
 
 	// Graph to run
-	UPROPERTY(EditAnywhere, Category = "Graph", meta = (Input))
+	UPROPERTY(EditAnywhere, Category = "Graph", meta = (Input, ExportAsReference="true"))
 	TObjectPtr<UAnimNextAnimationGraph> Graph;
 
 	// Instance used to hold graph state
diff --git a/Engine/Plugins/Experimental/AnimNextAnimGraph/Source/AnimNextAnimGraph/Private/Traits/SequencePlayer.h b/Engine/Plugins/Experimental/AnimNextAnimGraph/Source/AnimNextAnimGraph/Private/Traits/SequencePlayer.h
--- a/Engine/Plugins/Experimental/AnimNextAnimGraph/Source/AnimNextAnimGraph/Private/Traits/SequencePlayer.h
+++ b/Engine/Plugins/Experimental/AnimNextAnimGraph/Source/AnimNextAnimGraph/Private/Traits/SequencePlayer.h
@@ -18,7 +18,7 @@ struct FAnimNextSequencePlayerTraitSharedData : public FAnimNextTraitSharedData
 	GENERATED_BODY()
 
 	/** The sequence to play. */
-	UPROPERTY(EditAnywhere, Category = "Default")
+	UPROPERTY(EditAnywhere, Category = "Default", meta=(ExportAsReference="true"))
 	TObjectPtr<UAnimSequence> AnimSequence;
 
 	/** The play rate multiplier at which this sequence plays. */
diff --git a/Engine/Plugins/Experimental/AnimNextAnimGraph/Source/AnimNextAnimGraph/Private/Traits/SubGraphHost.h b/Engine/Plugins/Experimental/AnimNextAnimGraph/Source/AnimNextAnimGraph/Private/Traits/SubGraphHost.h
--- a/Engine/Plugins/Experimental/AnimNextAnimGraph/Source/AnimNextAnimGraph/Private/Traits/SubGraphHost.h
+++ b/Engine/Plugins/Experimental/AnimNextAnimGraph/Source/AnimNextAnimGraph/Private/Traits/SubGraphHost.h
@@ -20,7 +20,7 @@ struct FAnimNextSubGraphHostTraitSharedData : public FAnimNextTraitSharedData
 	GENERATED_BODY()
 
 	/** Asset to use as a sub-graph */
-	UPROPERTY(EditAnywhere, Category = "Default")
+	UPROPERTY(EditAnywhere, Category = "Default", meta=(ExportAsReference="true"))
 	TObjectPtr<const UAnimNextAnimationGraph> AnimationGraph;
 
 	/** A dummy child that we can use to output the bind pose. This property is hidden and automatically populated during compilation. */
diff --git a/Engine/Plugins/Experimental/AnimNextAnimGraph/Source/AnimNextAnimGraphUncookedOnly/Private/AnimNextAnimGraphWorkspaceAssetUserData.cpp b/Engine/Plugins/Experimental/AnimNextAnimGraph/Source/AnimNextAnimGraphUncookedOnly/Private/AnimNextAnimGraphWorkspaceAssetUserData.cpp
--- a/Engine/Plugins/Experimental/AnimNextAnimGraph/Source/AnimNextAnimGraphUncookedOnly/Private/AnimNextAnimGraphWorkspaceAssetUserData.cpp
+++ b/Engine/Plugins/Experimental/AnimNextAnimGraph/Source/AnimNextAnimGraphUncookedOnly/Private/AnimNextAnimGraphWorkspaceAssetUserData.cpp
@@ -21,7 +21,7 @@ void UAnimNextAnimGraphWorkspaceAssetUserData::GetAssetRegistryTags(FAssetRegist
 			Data.Asset = Asset;
 		}
 
-		UE::AnimNext::UncookedOnly::FUtils::GetAssetOutlinerItems(EditorData, Exports);
+		UE::AnimNext::UncookedOnly::FUtils::GetAssetOutlinerItems(EditorData, Exports, Context);
 	}
 
 	FString TagValue;
diff --git a/Engine/Plugins/Experimental/AnimNextChooser/Source/AnimNextChooser/Private/MakeGraphFromChooser.h b/Engine/Plugins/Experimental/AnimNextChooser/Source/AnimNextChooser/Private/MakeGraphFromChooser.h
--- a/Engine/Plugins/Experimental/AnimNextChooser/Source/AnimNextChooser/Private/MakeGraphFromChooser.h
+++ b/Engine/Plugins/Experimental/AnimNextChooser/Source/AnimNextChooser/Private/MakeGraphFromChooser.h
@@ -16,7 +16,7 @@ struct FAnimNextMakeGraphFromChooserTraitSharedData : public FAnimNextTraitShare
 	GENERATED_BODY()
 
 	// The chooser to use to pick an animation graph 
-	UPROPERTY(EditAnywhere, Category = "Chooser")
+	UPROPERTY(EditAnywhere, Category = "Chooser", meta=(ExportAsReference="true"))
 	TObjectPtr<UChooserTable> Chooser;
 
 	// Latent pin support boilerplate
diff --git a/Engine/Plugins/Experimental/AnimNextStateTree/Source/AnimNextStateTree/Internal/AnimStateTreeTrait.h b/Engine/Plugins/Experimental/AnimNextStateTree/Source/AnimNextStateTree/Internal/AnimStateTreeTrait.h
--- a/Engine/Plugins/Experimental/AnimNextStateTree/Source/AnimNextStateTree/Internal/AnimStateTreeTrait.h
+++ b/Engine/Plugins/Experimental/AnimNextStateTree/Source/AnimNextStateTree/Internal/AnimStateTreeTrait.h
@@ -19,7 +19,7 @@ struct FAnimNextStateTreeTraitSharedData : public FAnimNextTraitSharedData
 {
 	GENERATED_BODY()
 
-	UPROPERTY(EditAnywhere, Category = "Default")
+	UPROPERTY(EditAnywhere, Category = "Default", meta=(ExportAsReference="true"))
 	FStateTreeReference StateTreeReference;
 	
 	UPROPERTY(EditAnywhere, Category = "Default")
diff --git a/Engine/Plugins/Experimental/AnimNextStateTree/Source/AnimNextStateTree/Private/AnimNextStateTreeGraphInstanceTask.cpp b/Engine/Plugins/Experimental/AnimNextStateTree/Source/AnimNextStateTree/Private/AnimNextStateTreeGraphInstanceTask.cpp
--- a/Engine/Plugins/Experimental/AnimNextStateTree/Source/AnimNextStateTree/Private/AnimNextStateTreeGraphInstanceTask.cpp
+++ b/Engine/Plugins/Experimental/AnimNextStateTree/Source/AnimNextStateTree/Private/AnimNextStateTreeGraphInstanceTask.cpp
@@ -77,4 +77,12 @@ void FAnimNextStateTreeGraphInstanceTask::PostEditInstanceDataChangeChainPropert
 		GetDefault<UAnimNextAnimGraphSettings>()->GetNonNativePayloadFromGraph(InstanceData.Asset, AnimationGraph, InstanceData.Payload);
 	}
 }
-#endif
\ No newline at end of file
+
+void FAnimNextStateTreeGraphInstanceTask::GetObjectReferences(TArray<const UObject*>& OutReferencedObjects, FStateTreeDataView InstanceDataView) const
+{
+	if (InstanceDataView.IsValid())
+	{
+		OutReferencedObjects.Add(InstanceDataView.Get<FAnimNextGraphInstanceTaskInstanceData>().Asset.Get());
+	}	
+}
+#endif
diff --git a/Engine/Plugins/Experimental/AnimNextStateTree/Source/AnimNextStateTree/Public/AnimNextStateTreeTypes.h b/Engine/Plugins/Experimental/AnimNextStateTree/Source/AnimNextStateTree/Public/AnimNextStateTreeTypes.h
--- a/Engine/Plugins/Experimental/AnimNextStateTree/Source/AnimNextStateTree/Public/AnimNextStateTreeTypes.h
+++ b/Engine/Plugins/Experimental/AnimNextStateTree/Source/AnimNextStateTree/Public/AnimNextStateTreeTypes.h
@@ -24,4 +24,8 @@ USTRUCT(meta = (DisplayName = "AnimNext Task Base", Hidden))
 struct ANIMNEXTSTATETREE_API FAnimNextStateTreeTaskBase : public FStateTreeTaskBase
 {
 	GENERATED_BODY()
+
+#if WITH_EDITOR
+	virtual void GetObjectReferences(TArray<const UObject*>& OutReferencedObjects, const FStateTreeDataView InstanceDataView) const {} 
+#endif // WITH_EDITOR
 };
diff --git a/Engine/Plugins/Experimental/AnimNextStateTree/Source/AnimNextStateTree/Public/Tasks/AnimNextStateTreeGraphInstanceTask.h b/Engine/Plugins/Experimental/AnimNextStateTree/Source/AnimNextStateTree/Public/Tasks/AnimNextStateTreeGraphInstanceTask.h
--- a/Engine/Plugins/Experimental/AnimNextStateTree/Source/AnimNextStateTree/Public/Tasks/AnimNextStateTreeGraphInstanceTask.h
+++ b/Engine/Plugins/Experimental/AnimNextStateTree/Source/AnimNextStateTree/Public/Tasks/AnimNextStateTreeGraphInstanceTask.h
@@ -54,6 +54,7 @@ protected:
 	virtual void ExitState(FStateTreeExecutionContext& Context, const FStateTreeTransitionResult& Transition) const override;
 #if WITH_EDITOR
 	virtual void PostEditInstanceDataChangeChainProperty(const FPropertyChangedChainEvent& PropertyChangedEvent, FStateTreeDataView InstanceDataView) override;
+	virtual void GetObjectReferences(TArray<const UObject*>& OutReferencedObjects, const FStateTreeDataView InstanceDataView) const override;
 #endif
 public:
 	TStateTreeExternalDataHandle<FAnimNextStateTreeTraitContext> TraitContextHandle;
diff --git a/Engine/Plugins/Experimental/AnimNextStateTree/Source/AnimNextStateTreeEditor/Private/AnimNextStateTreeEditorModule.cpp b/Engine/Plugins/Experimental/AnimNextStateTree/Source/AnimNextStateTreeEditor/Private/AnimNextStateTreeEditorModule.cpp
--- a/Engine/Plugins/Experimental/AnimNextStateTree/Source/AnimNextStateTreeEditor/Private/AnimNextStateTreeEditorModule.cpp
+++ b/Engine/Plugins/Experimental/AnimNextStateTree/Source/AnimNextStateTreeEditor/Private/AnimNextStateTreeEditorModule.cpp
@@ -17,10 +17,12 @@
 #include "AnimNextStateTreeEditorHost.h"
 #include "StateTreeEditingSubsystem.h"
 #include "StateTreeEditorStyle.h"
+#include "UncookedOnlyUtils.h"
 #include "Framework/Docking/LayoutExtender.h"
 #include "Toolkits/AssetEditorModeUILayer.h"
 #include "ToolMenus.h"
 #include "Framework/MultiBox/ToolMenuBase.h"
+#include "AnimNextRigVMAssetEditorData.h"
 
 #define LOCTEXT_NAMESPACE "AnimNextStateTreeEditorModule"
 
@@ -173,6 +175,102 @@ void FAnimNextStateTreeEditorModule::StartupModule()
 			
 			return FSlateColor::UseForeground();
 		}
+
+		static void SelectStateExports(TConstArrayView<FWorkspaceOutlinerItemExport> StateExports)
+		{
+			if(UStateTreeEditingSubsystem* EditingSubsystem = GEditor->GetEditorSubsystem<UStateTreeEditingSubsystem>())
+			{
+				UAnimNextStateTree* SelectionStateTree = [StateExports]() -> UAnimNextStateTree*
+				{
+					UAnimNextStateTree* FoundTree = nullptr;
+
+					for (const FWorkspaceOutlinerItemExport& SelectedExport : StateExports)
+					{
+						if (UAnimNextStateTree* LoadedStateTree = Cast<UAnimNextStateTree>(SelectedExport.GetFirstAssetPath().ResolveObject()))
+						{
+							if (FoundTree == nullptr)
+							{
+								FoundTree = LoadedStateTree;
+							}
+							if (FoundTree != LoadedStateTree)
+							{
+								FoundTree = nullptr;
+								break;
+							}
+						}
+					}
+
+					return FoundTree;
+				}();
+
+				if (SelectionStateTree)
+				{
+					TSharedPtr<FStateTreeViewModel> ViewModel = EditingSubsystem->FindOrAddViewModel(SelectionStateTree->StateTree.Get());
+					check(ViewModel.IsValid());
+			
+					TArray<TWeakObjectPtr<UStateTreeState>> ToBeSelectedStates;
+					for (const FWorkspaceOutlinerItemExport& SelectedExport : StateExports)
+					{
+						if (SelectedExport.GetData().GetScriptStruct() == FAnimNextStateTreeStateOutlinerData::StaticStruct())
+						{
+							if (UAnimNextStateTree* LoadedStateTree = Cast<UAnimNextStateTree>(SelectedExport.GetFirstAssetPath().ResolveObject()))
+							{
+								const FAnimNextStateTreeStateOutlinerData& StateData = SelectedExport.GetData().Get<FAnimNextStateTreeStateOutlinerData>();		
+								ensure(LoadedStateTree == SelectionStateTree);
+								if (UStateTreeState* State = ViewModel->GetMutableStateByID(StateData.StateId))
+								{
+									ToBeSelectedStates.Add(State);
+								}
+							}
+						}
+					}
+
+					ViewModel->SetSelection(ToBeSelectedStates);
+				}
+			}	
+		}
+		
+		virtual bool HandleDoubleClick(const FToolMenuContext& ToolMenuContext) const override
+		{
+			const UWorkspaceItemMenuContext* WorkspaceItemContext = ToolMenuContext.FindContext<UWorkspaceItemMenuContext>();
+			const UAssetEditorToolkitMenuContext* AssetEditorContext = ToolMenuContext.FindContext<UAssetEditorToolkitMenuContext>(); 
+			if (WorkspaceItemContext && AssetEditorContext)
+			{				
+				if(const TSharedPtr<UE::Workspace::IWorkspaceEditor> WorkspaceEditor = StaticCastSharedPtr<UE::Workspace::IWorkspaceEditor>(AssetEditorContext->Toolkit.Pin()))
+				{
+					const uint32 NumSelectedExports = WorkspaceItemContext->SelectedExports.Num();				
+					if (NumSelectedExports == 1)
+					{
+						const FWorkspaceOutlinerItemExport& SelectedExport = WorkspaceItemContext->SelectedExports[0];
+
+						if(WorkspaceItemContext->SelectedExports[0].GetData().GetScriptStruct() == FAnimNextStateTreeOutlinerData::StaticStruct())
+						{
+							if (UAnimNextStateTree* LoadedStateTree = SelectedExport.GetFirstAssetOfType<UAnimNextStateTree>())
+							{
+								WorkspaceEditor->OpenAssets({LoadedStateTree });
+
+								return true;
+							}
+						}
+						else if (SelectedExport.GetData().GetScriptStruct() == FAnimNextStateTreeStateOutlinerData::StaticStruct())
+						{
+							if (UAnimNextStateTree* SelectionStateTree = SelectedExport.GetFirstAssetOfType<UAnimNextStateTree>())
+							{
+								// First open editor for state-tree
+								WorkspaceEditor->OpenAssets({SelectionStateTree });
+
+								// Then select the double-clicked state
+								SelectStateExports(WorkspaceItemContext->SelectedExports);
+								
+								return true;
+							}								
+						}							
+					}					
+				}			
+			}
+
+			return false;
+		}
 		
 		virtual bool HandleSelected(const FToolMenuContext& ToolMenuContext) const override
 		{
@@ -189,64 +287,14 @@ void FAnimNextStateTreeEditorModule::StartupModule()
 						if (NumSelectedExports == 1 && WorkspaceItemContext->SelectedExports[0].GetData().GetScriptStruct() == FAnimNextStateTreeOutlinerData::StaticStruct())
 						{
 							const FWorkspaceOutlinerItemExport& SelectedExport = WorkspaceItemContext->SelectedExports[0];
-							if (UAnimNextStateTree* LoadedStateTree = Cast<UAnimNextStateTree>(SelectedExport.GetAssetPath().ResolveObject()))
+							if (UAnimNextStateTree* LoadedStateTree = Cast<UAnimNextStateTree>(SelectedExport.GetFirstAssetPath().ResolveObject()))
 							{
 								WorkspaceEditor->SetDetailsObjects({ LoadedStateTree->StateTree->EditorData } );
 								return true;
 							}
 						}
 
-						if(UStateTreeEditingSubsystem* EditingSubsystem = GEditor->GetEditorSubsystem<UStateTreeEditingSubsystem>())
-						{
-							UAnimNextStateTree* SelectionStateTree = [WorkspaceItemContext]() -> UAnimNextStateTree*
-							{
-								UAnimNextStateTree* FoundTree = nullptr;
-
-								for (const FWorkspaceOutlinerItemExport& SelectedExport : WorkspaceItemContext->SelectedExports)
-								{
-									if (UAnimNextStateTree* LoadedStateTree = Cast<UAnimNextStateTree>(SelectedExport.GetAssetPath().ResolveObject()))
-									{
-										if (FoundTree == nullptr)
-										{
-											FoundTree = LoadedStateTree;
-										}
-										if (FoundTree != LoadedStateTree)
-										{
-											FoundTree = nullptr;
-											break;
-										}
-									}
-								}
-
-								return FoundTree;
-							}();
-
-							if (SelectionStateTree)
-							{
-								TSharedPtr<FStateTreeViewModel> ViewModel = EditingSubsystem->FindOrAddViewModel(SelectionStateTree->StateTree.Get());
-								check(ViewModel.IsValid());
-								
-								TArray<TWeakObjectPtr<UStateTreeState>> ToBeSelectedStates;
-								for (const FWorkspaceOutlinerItemExport& SelectedExport : WorkspaceItemContext->SelectedExports)
-								{
-									if (SelectedExport.GetData().GetScriptStruct() == FAnimNextStateTreeStateOutlinerData::StaticStruct())
-									{
-										if (UAnimNextStateTree* LoadedStateTree = Cast<UAnimNextStateTree>(SelectedExport.GetAssetPath().ResolveObject()))
-										{
-											const FAnimNextStateTreeStateOutlinerData& StateData = SelectedExport.GetData().Get<FAnimNextStateTreeStateOutlinerData>();		
-											ensure(LoadedStateTree == SelectionStateTree);
-											if (UStateTreeState* State = ViewModel->GetMutableStateByID(StateData.StateId))
-											{
-												ToBeSelectedStates.Add(State);
-											}
-										}
-									}
-								}
-
-								ViewModel->SetSelection(ToBeSelectedStates);
-								return true;
-							}
-						}						
+						SelectStateExports(WorkspaceItemContext->SelectedExports);
 					}
 				}			
 			}
diff --git a/Engine/Plugins/Experimental/AnimNextStateTree/Source/AnimNextStateTreeUncookedOnly/Private/AnimNextStateTreeWorkspaceAssetUserData.cpp b/Engine/Plugins/Experimental/AnimNextStateTree/Source/AnimNextStateTreeUncookedOnly/Private/AnimNextStateTreeWorkspaceAssetUserData.cpp
--- a/Engine/Plugins/Experimental/AnimNextStateTree/Source/AnimNextStateTreeUncookedOnly/Private/AnimNextStateTreeWorkspaceAssetUserData.cpp
+++ b/Engine/Plugins/Experimental/AnimNextStateTree/Source/AnimNextStateTreeUncookedOnly/Private/AnimNextStateTreeWorkspaceAssetUserData.cpp
@@ -5,7 +5,9 @@
 #include "AnimNextStateTree.h"
 #include "StateTreeEditorData.h"
 #include "AnimNextStateTreeWorkspaceExports.h"
+#include "Tasks/AnimNextStateTreeGraphInstanceTask.h"
 #include "UObject/AssetRegistryTagsContext.h"
+#include "WorkspaceAssetRegistryInfo.h"
 
 void UAnimNextStateTreeWorkspaceAssetUserData::GetAssetRegistryTags(FAssetRegistryTagsContext Context) const
 {
@@ -45,6 +47,31 @@ void UAnimNextStateTreeWorkspaceAssetUserData::GetAssetRegistryTags(FAssetRegist
 				const FStateTreeEditorColor* StateTreeEditorColor = EditorData->FindColor(State.ColorRef);
 
 				AssetData.Color = StateTreeEditorColor ? FSlateColor(StateTreeEditorColor->Color) : FSlateColor::UseForeground();
+
+				// Retrieve and add AssetReferences for all tasks part of this state
+				{
+					TArray<const UObject*> ReferencedObjects;
+					for (FStateTreeEditorNode& TaskNode : State.Tasks)
+					{
+						if(TaskNode.Node.IsValid() && TaskNode.Node.GetScriptStruct()->IsChildOf(FAnimNextStateTreeTaskBase::StaticStruct()))
+						{
+							if (TaskNode.Instance.IsValid())
+							{							
+								TaskNode.Node.Get<FAnimNextStateTreeTaskBase>().GetObjectReferences(ReferencedObjects, FStateTreeDataView(TaskNode.Instance.GetScriptStruct(), TaskNode.Instance.GetMutableMemory()));
+							}
+						}
+					}
+
+					for (const UObject* ReferredObject : ReferencedObjects)
+					{
+						FSoftObjectPath Path = ReferredObject;
+						FWorkspaceOutlinerItemExport& GraphReference = OutlinerExports.Exports.Add_GetRef(FWorkspaceOutlinerItemExport(FName(Path.ToString()), StateExport));
+						GraphReference.GetData().InitializeAsScriptStruct(FWorkspaceOutlinerAssetReferenceItemData::StaticStruct());
+						FWorkspaceOutlinerAssetReferenceItemData& ReferenceData = GraphReference.GetData().GetMutable<FWorkspaceOutlinerAssetReferenceItemData>();
+						ReferenceData.ReferredObjectPath = Path;
+					}
+				}
+				
 				
 				return EStateTreeVisitor::Continue;
 			});
diff --git a/Engine/Plugins/Experimental/Workspace/Source/WorkspaceEditor/Private/Outliner/WorkspaceOutlinerColumns.cpp b/Engine/Plugins/Experimental/Workspace/Source/WorkspaceEditor/Private/Outliner/WorkspaceOutlinerColumns.cpp
--- a/Engine/Plugins/Experimental/Workspace/Source/WorkspaceEditor/Private/Outliner/WorkspaceOutlinerColumns.cpp
+++ b/Engine/Plugins/Experimental/Workspace/Source/WorkspaceEditor/Private/Outliner/WorkspaceOutlinerColumns.cpp
@@ -59,7 +59,7 @@ namespace UE::Workspace
 				
 				if (!WeakPackage.IsValid())
 				{
-					if(const UObject* Object = TreeItem->Export.GetAssetPath().ResolveObject())
+					if(const UObject* Object = TreeItem->Export.GetFirstAssetPath().ResolveObject())
 					{
 						WeakPackage = Object->GetPackage();
 					}
diff --git a/Engine/Plugins/Experimental/Workspace/Source/WorkspaceEditor/Private/Outliner/WorkspaceOutlinerHierarchy.cpp b/Engine/Plugins/Experimental/Workspace/Source/WorkspaceEditor/Private/Outliner/WorkspaceOutlinerHierarchy.cpp
--- a/Engine/Plugins/Experimental/Workspace/Source/WorkspaceEditor/Private/Outliner/WorkspaceOutlinerHierarchy.cpp
+++ b/Engine/Plugins/Experimental/Workspace/Source/WorkspaceEditor/Private/Outliner/WorkspaceOutlinerHierarchy.cpp
@@ -6,6 +6,8 @@
 #include "Workspace.h"
 #include "WorkspaceAssetRegistryInfo.h"
 #include "WorkspaceOutlinerTreeItem.h"
+#include "AssetRegistry/AssetRegistryModule.h"
+#include "AssetRegistry/IAssetRegistry.h"
 
 namespace UE::Workspace
 {
@@ -13,28 +15,139 @@ namespace UE::Workspace
 	{
 	}
 
-	void FWorkspaceOutlinerHierarchy::CreateItems(TArray<FSceneOutlinerTreeItemPtr>& OutItems) const
+	void FWorkspaceOutlinerHierarchy::CreateItemsFromAssetData(const FAssetData& AssetData, TArray<FSceneOutlinerTreeItemPtr>& OutItems, TArray<FAssetData>& OutAssets, const FWorkspaceOutlinerItemExport* InParentExport) const
 	{
-		if (const UWorkspace* Workspace = WeakWorkspace.Get())
+		FString TagValue;
+		if(AssetData.GetTagValue(UE::Workspace::ExportsWorkspaceItemsRegistryTag, TagValue))
 		{
-			TArray<FAssetData> AssetDataEntries; 
-			Workspace->GetAssetDataEntries(AssetDataEntries);
+			FWorkspaceOutlinerItemExports Exports;
+			FWorkspaceOutlinerItemExports::StaticStruct()->ImportText(*TagValue, &Exports, nullptr, 0, nullptr, FWorkspaceOutlinerItemExports::StaticStruct()->GetName());
 
-			for (const FAssetData& AssetData : AssetDataEntries)
+			const bool bRecursiveReference = OutAssets.Contains(AssetData);
+
+			// Reparented hashes, required to make asset-references (and their child hierarchies) unique inside of the workspace
+			TMap<uint32, FWorkspaceOutlinerItemExport> ReparentedHashes;			
+			for (const FWorkspaceOutlinerItemExport& Export : Exports.Exports)
 			{
-				FString TagValue;
-				if(AssetData.GetTagValue(UE::Workspace::ExportsWorkspaceItemsRegistryTag, TagValue))
+				auto AddItem = [this, &OutItems](const FWorkspaceOutlinerItemExport& Export)
+				{
+					if (FSceneOutlinerTreeItemPtr Item = Mode->CreateItemFor<FWorkspaceOutlinerTreeItem>(FWorkspaceOutlinerTreeItem::FItemData{Export}))
+					{
+						OutItems.Add(Item);
+					}
+				};
+
+				const bool bIsReference = Export.GetData().GetScriptStruct() == FWorkspaceOutlinerAssetReferenceItemData::StaticStruct();
+
+				// Early-out circular/recursive references
+				if (bIsReference && bRecursiveReference)
+				{
+					continue;
+				}
+
+				if (bIsReference)
+				{
+					const FWorkspaceOutlinerAssetReferenceItemData& ReferenceItemData = Export.GetData().Get<FWorkspaceOutlinerAssetReferenceItemData>();		
+
+					// Retrieve AssetRegistry data for the referred asset
+					FARFilter Filter;
+					Filter.SoftObjectPaths.Add(ReferenceItemData.ReferredObjectPath);
+					const IAssetRegistry& AssetRegistry = FAssetRegistryModule::GetRegistry();
+					TArray<FAssetData> OutAssetData;
+					AssetRegistry.GetAssets(Filter, OutAssetData);
+
+					if (OutAssetData.Num())
+					{
+						// Only ever expect one to be found
+						check(OutAssetData.Num() == 1);
+					
+						const FAssetData& ReferredAssetData = OutAssetData[0];
+						TArray<FSceneOutlinerTreeItemPtr> ReferredItems;
+
+
+						// Make local copy of AssetData array for recursion checks
+						TArray<FAssetData> ReferenceAssets = OutAssets;
+						ReferenceAssets.Add(AssetData);
+					
+						if (InParentExport)
+						{
+							// Asset reference not at root-level of the workspace, meaning its root entry has to be reparented (in case the find fails, the asset registry exporting is not being done in the same order as outliner representation)
+							// Populate export to-be-used by root-level asset referenced export 
+							const FWorkspaceOutlinerItemExport ReparentedExport(Export.GetIdentifier(), ReparentedHashes.FindChecked(Export.GetParentHash()), Export.GetData());						
+							ReparentedHashes.Add(GetTypeHash(Export), ReparentedExport);
+							CreateItemsFromAssetData(ReferredAssetData, ReferredItems, ReferenceAssets, &ReparentedExport);
+						}
+						else
+						{
+							// Reference in root-level workspace asset
+							CreateItemsFromAssetData(ReferredAssetData, ReferredItems,ReferenceAssets,  &Export);
+						}
+
+						if(ReferredItems.Num())
+						{
+							OutItems.Append(ReferredItems);
+						}
+					}
+				
+				}
+
+				// Handle exports for referenced assets
+				if (InParentExport)
 				{
-					FWorkspaceOutlinerItemExports Exports;
-					FWorkspaceOutlinerItemExports::StaticStruct()->ImportText(*TagValue, &Exports, nullptr, 0, nullptr, FWorkspaceOutlinerItemExports::StaticStruct()->GetName());
-					for (const FWorkspaceOutlinerItemExport& Export : Exports.Exports)
+					// Root export for asset
+					if (Export.GetParentIdentifier() == NAME_None)
 					{
-						if (FSceneOutlinerTreeItemPtr Item = Mode->CreateItemFor<FWorkspaceOutlinerTreeItem>(FWorkspaceOutlinerTreeItem::FItemData{Export}))
+						FWorkspaceOutlinerItemExport ReferenceExport = *InParentExport;
+						ensure(ReferenceExport.GetData().GetScriptStruct() == FWorkspaceOutlinerAssetReferenceItemData::StaticStruct());
+						
+						// Save a copy of the original assets root export on the (to-be-added) asset reference export data
+						FWorkspaceOutlinerAssetReferenceItemData& ItemData = ReferenceExport.GetData().GetMutable<FWorkspaceOutlinerAssetReferenceItemData>();
+						ItemData.ReferredExport = Export;
+						ItemData.bRecursiveReference = bRecursiveReference;
+
+						// Add this new reference export as new root with parent hash of original root asset export
+						ReparentedHashes.Add(GetTypeHash(Export), ReferenceExport);
+
+						AddItem(ReferenceExport);
+					}
+					else
+					{
+						// Do not add entries other than root export for recursive references
+						if (OutAssets.Contains(AssetData))
 						{
-							OutItems.Add(Item);
+							continue;
 						}
+
+						// in case the find fails, the asset registry exporting is not being done in the same order as outliner representation
+						const FWorkspaceOutlinerItemExport ReparentedExport(Export.GetIdentifier(), ReparentedHashes.FindChecked(Export.GetParentHash()), Export.GetData());
+						ReparentedHashes.Add(GetTypeHash(Export), ReparentedExport);
+
+						AddItem(ReparentedExport);
 					}
 				}
+				else
+				{
+					AddItem(Export);				
+				}
+
+				
+			}
+			
+			OutAssets.Add(AssetData);
+		}
+	}
+
+	void FWorkspaceOutlinerHierarchy::CreateItems(TArray<FSceneOutlinerTreeItemPtr>& OutItems) const
+	{
+		if (const UWorkspace* Workspace = WeakWorkspace.Get())
+		{
+			TArray<FAssetData> AssetDataEntries; 
+			Workspace->GetAssetDataEntries(AssetDataEntries);
+
+			for (const FAssetData& AssetData : AssetDataEntries)
+			{
+				TArray<FAssetData> Assets;
+				CreateItemsFromAssetData(AssetData, OutItems, Assets, nullptr);
 			}
 		}		
 	}
@@ -57,7 +170,7 @@ namespace UE::Workspace
 						TArray<FAssetData> AssetDataEntries; 
 						Workspace->GetAssetDataEntries(AssetDataEntries);
 
-						if(const FAssetData* AssetDataPtr = AssetDataEntries.FindByPredicate([AssetPath = TreeItem->Export.GetAssetPath()](const FAssetData& AssetData) { return AssetData.GetSoftObjectPath() == AssetPath; }))
+						if(const FAssetData* AssetDataPtr = AssetDataEntries.FindByPredicate([AssetPath = TreeItem->Export.GetTopLevelAssetPath()](const FAssetData& AssetData) { return AssetData.GetSoftObjectPath() == AssetPath; }))
 						{
 							FString TagValue;
 							if(AssetDataPtr->GetTagValue(UE::Workspace::ExportsWorkspaceItemsRegistryTag, TagValue))
diff --git a/Engine/Plugins/Experimental/Workspace/Source/WorkspaceEditor/Private/Outliner/WorkspaceOutlinerHierarchy.h b/Engine/Plugins/Experimental/Workspace/Source/WorkspaceEditor/Private/Outliner/WorkspaceOutlinerHierarchy.h
--- a/Engine/Plugins/Experimental/Workspace/Source/WorkspaceEditor/Private/Outliner/WorkspaceOutlinerHierarchy.h
+++ b/Engine/Plugins/Experimental/Workspace/Source/WorkspaceEditor/Private/Outliner/WorkspaceOutlinerHierarchy.h
@@ -6,6 +6,7 @@
 #include "ISceneOutlinerHierarchy.h"
 
 class UWorkspace;
+struct FWorkspaceOutlinerItemExport;
 
 namespace UE::Workspace
 {
@@ -21,6 +22,8 @@ namespace UE::Workspace
 		virtual void CreateChildren(const FSceneOutlinerTreeItemPtr& Item, TArray<FSceneOutlinerTreeItemPtr>& OutChildren) const override {}
 		virtual FSceneOutlinerTreeItemPtr FindOrCreateParentItem(const ISceneOutlinerTreeItem& Item, const TMap<FSceneOutlinerTreeItemID, FSceneOutlinerTreeItemPtr>& Items, bool bCreate = false) override;
 		// End ISceneOutlinerHierarchy overrides
+	protected:
+		void CreateItemsFromAssetData(const FAssetData& AssetData, TArray<FSceneOutlinerTreeItemPtr>& OutItems, TArray<FAssetData>& OutAssets, const FWorkspaceOutlinerItemExport* InParentExport = nullptr) const;
 	private:
 		TWeakObjectPtr<UWorkspace> WeakWorkspace;
 	};
diff --git a/Engine/Plugins/Experimental/Workspace/Source/WorkspaceEditor/Private/Outliner/WorkspaceOutlinerMode.cpp b/Engine/Plugins/Experimental/Workspace/Source/WorkspaceEditor/Private/Outliner/WorkspaceOutlinerMode.cpp
--- a/Engine/Plugins/Experimental/Workspace/Source/WorkspaceEditor/Private/Outliner/WorkspaceOutlinerMode.cpp
+++ b/Engine/Plugins/Experimental/Workspace/Source/WorkspaceEditor/Private/Outliner/WorkspaceOutlinerMode.cpp
@@ -77,13 +77,13 @@ TSharedPtr<SWidget> FWorkspaceOutlinerMode::CreateContextMenu()
 					CommonSection.AddMenuEntryWithCommandList(FGenericCommands::Get().Rename, MenuContext->WeakCommandList.Pin());
 					// TODO JDB more possible actions? (find in content browser?)
 
-					const bool bSelectionContainsTopLevelAsset = MenuContext->SelectedExports.Num() && MenuContext->SelectedExports.ContainsByPredicate([](const FWorkspaceOutlinerItemExport& Export)
+					const bool bSelectionContainsValidAssetPath = MenuContext->SelectedExports.Num() && MenuContext->SelectedExports.ContainsByPredicate([](const FWorkspaceOutlinerItemExport& Export)
 					{
-						return Export.GetParentIdentifier() == NAME_None;
+						return Export.GetFirstAssetPath().IsValid();
 					});
 					
 					FToolMenuSection& AssetsSection = InMenu->AddSection("Assets", LOCTEXT("AssetSectionLabel", "Assets"));
-					if (bSelectionContainsTopLevelAsset)
+					if (bSelectionContainsValidAssetPath)
 					{
 						AssetsSection.AddMenuEntry(TEXT("BrowseToAsset"),
 							LOCTEXT("BrowseToAssetLabel", "Browse to Asset"),
@@ -95,10 +95,7 @@ TSharedPtr<SWidget> FWorkspaceOutlinerMode::CreateContextMenu()
 									FARFilter Filter;	
 									for (const FWorkspaceOutlinerItemExport& ItemExport : SelectedExports)
 									{
-										if (ItemExport.GetParentIdentifier() == NAME_None)
-										{
-											Filter.SoftObjectPaths.AddUnique(ItemExport.GetAssetPath());
-										}
+										Filter.SoftObjectPaths.AddUnique(ItemExport.GetFirstAssetPath());
 									}
 
 									TArray<FAssetData> AssetDataList;
@@ -136,9 +133,9 @@ TSharedPtr<SWidget> FWorkspaceOutlinerMode::CreateContextMenu()
 											SavablePackages.AddUnique(Package);								
 										}										
 									}
-									else if (ItemExport.GetParentIdentifier() == NAME_None)
+									else
 									{
-										if (UPackage* Package = FindPackage(nullptr, *ItemExport.GetAssetPath().GetLongPackageName()))
+										if (UPackage* Package = FindPackage(nullptr, *ItemExport.GetFirstAssetPath().GetLongPackageName()))
 										{
 											if (IsPackageDirty(Package))
 											{
@@ -168,7 +165,7 @@ TSharedPtr<SWidget> FWorkspaceOutlinerMode::CreateContextMenu()
 									}
 									else if (ItemExport.GetParentIdentifier() == NAME_None)
 									{
-										if (const UPackage* Package = FindPackage(nullptr, *ItemExport.GetAssetPath().GetLongPackageName()))
+										if (const UPackage* Package = FindPackage(nullptr, *ItemExport.GetFirstAssetPath().GetLongPackageName()))
 										{
 											if (IsPackageDirty(Package))
 											{
@@ -251,7 +248,13 @@ void FWorkspaceOutlinerMode::HandleItemSelection(const FSceneOutlinerItemSelecti
 					},
 					[](const FSceneOutlinerTreeItemPtr& InItem)
 					{
-						return InItem->CastTo<FWorkspaceOutlinerTreeItem>()->Export;
+						FWorkspaceOutlinerTreeItem* TreeItem = InItem->CastTo<FWorkspaceOutlinerTreeItem>();
+						if (TreeItem->Export.HasData() && TreeItem->Export.GetData().GetScriptStruct() == FWorkspaceOutlinerAssetReferenceItemData::StaticStruct())
+						{
+							return TreeItem->Export.GetData().GetMutable<FWorkspaceOutlinerAssetReferenceItemData>().ReferredExport;
+						}
+
+						return TreeItem->Export;
 					});
 
 			// Singular item selected
@@ -269,20 +272,20 @@ void FWorkspaceOutlinerMode::HandleItemSelection(const FSceneOutlinerItemSelecti
 					bHandled = SharedDetails->HandleSelected(Context);
 				}
 				
-				if(!bHandled && SelectedExport.GetData().IsValid())
+				if(!bHandled && SelectedExport.HasData())
 				{
 					TSharedPtr<FStructOnScope> ExportDataView = MakeShared<FStructOnScope>(SelectedExport.GetData().GetScriptStruct(), SelectedExport.GetData().GetMutableMemory());
 					// TODO JDB handle struct selections
 					//SharedWorkspaceEditor->SetDetailsStruct(ExportDataView);
 				}
 
-				UObject* LoadedAsset = SelectedExport.GetAssetPath().ResolveObject();
+				UObject* LoadedAsset = SelectedExport.GetTopLevelAssetPath().ResolveObject();
 				if (!bHandled && LoadedAsset)
 				{
 					SharedWorkspaceEditor->SetDetailsObjects({LoadedAsset});
 				}
 			}
-			else
+			else if (MenuContext->SelectedExports.Num() > 1)
 			{
 				const UScriptStruct* DataType = MenuContext->SelectedExports[0].GetData().GetScriptStruct();
 				Algo::ForEach(MenuContext->SelectedExports, [&DataType](const FWorkspaceOutlinerItemExport& Export)
@@ -398,24 +401,35 @@ void FWorkspaceOutlinerMode::OpenItems(TArrayView<const FSceneOutlinerTreeItemPt
 	{
 		if (const FWorkspaceOutlinerTreeItem* TreeItem = Item->CastTo<FWorkspaceOutlinerTreeItem>())
 		{
-			if(TreeItem->Export.GetParentIdentifier() == NAME_None)
+			bool bHandled = false;
+
+			FWorkspaceOutlinerItemExport Export = [&TreeItem]()
 			{
-				if (const TSharedPtr<UE::Workspace::IWorkspaceEditor> SharedWorkspaceEditor = StaticCastSharedPtr<UE::Workspace::IWorkspaceEditor>(WeakWorkspaceEditor.Pin()))
+				const FWorkspaceOutlinerTreeItem* OutlinerTreeItem = TreeItem->CastTo<FWorkspaceOutlinerTreeItem>();
+				if (OutlinerTreeItem->Export.HasData() && OutlinerTreeItem->Export.GetData().GetScriptStruct() == FWorkspaceOutlinerAssetReferenceItemData::StaticStruct())
 				{
-					SharedWorkspaceEditor->OpenAssets({TreeItem->Export.GetAssetPath().TryLoad()});
+					return OutlinerTreeItem->Export.GetData().Get<FWorkspaceOutlinerAssetReferenceItemData>().ReferredExport;
 				}
-			}
-			else
+
+				return TreeItem->Export;
+			}();
+			
+			if (const TSharedPtr<IWorkspaceOutlinerItemDetails> SharedDetails = FWorkspaceEditorModule::GetOutlinerItemDetails(MakeOutlinerDetailsId(Export)))
 			{
-				if (const TSharedPtr<IWorkspaceOutlinerItemDetails> SharedDetails = FWorkspaceEditorModule::GetOutlinerItemDetails(MakeOutlinerDetailsId(TreeItem->Export)))
-				{
-					UWorkspaceItemMenuContext* MenuContext = NewObject<UWorkspaceItemMenuContext>();
-					MenuContext->SelectedExports.Add(TreeItem->Export);
+				UWorkspaceItemMenuContext* MenuContext = NewObject<UWorkspaceItemMenuContext>();
+				MenuContext->SelectedExports.Add(Export);
 				
-					FToolMenuContext Context(MenuContext);
-					WeakWorkspaceEditor.Pin()->InitToolMenuContext(Context);
+				FToolMenuContext Context(MenuContext);
+				WeakWorkspaceEditor.Pin()->InitToolMenuContext(Context);
 				
-					SharedDetails->HandleDoubleClick(Context);
+				bHandled = SharedDetails->HandleDoubleClick(Context);
+			}
+			
+			if(!bHandled)
+			{
+				if (const TSharedPtr<UE::Workspace::IWorkspaceEditor> SharedWorkspaceEditor = StaticCastSharedPtr<UE::Workspace::IWorkspaceEditor>(WeakWorkspaceEditor.Pin()))
+				{
+					SharedWorkspaceEditor->OpenAssets({TreeItem->Export.GetFirstAssetPath().TryLoad()});
 				}
 			}
 		}
@@ -461,7 +475,7 @@ void FWorkspaceOutlinerMode::Delete()
 
 		if(TreeItem->Export.GetParentIdentifier() == NAME_None)
 		{
-			Filter.SoftObjectPaths.AddUnique(TreeItem->Export.GetAssetPath());
+			Filter.SoftObjectPaths.AddUnique(TreeItem->Export.GetFirstAssetPath());
 		}
 		else
 		{
diff --git a/Engine/Plugins/Experimental/Workspace/Source/WorkspaceEditor/Private/Outliner/WorkspaceOutlinerTreeItem.cpp b/Engine/Plugins/Experimental/Workspace/Source/WorkspaceEditor/Private/Outliner/WorkspaceOutlinerTreeItem.cpp
--- a/Engine/Plugins/Experimental/Workspace/Source/WorkspaceEditor/Private/Outliner/WorkspaceOutlinerTreeItem.cpp
+++ b/Engine/Plugins/Experimental/Workspace/Source/WorkspaceEditor/Private/Outliner/WorkspaceOutlinerTreeItem.cpp
@@ -133,6 +133,10 @@ namespace UE::Workspace
 	FWorkspaceOutlinerTreeItem::FWorkspaceOutlinerTreeItem(const FItemData& InItemData) : ISceneOutlinerTreeItem(FWorkspaceOutlinerTreeItem::Type), Export(InItemData.Export)
 	{
 		ItemDetails = FWorkspaceEditorModule::GetOutlinerItemDetails(MakeOutlinerDetailsId(Export));
+		if (ItemDetails.IsValid())
+		{
+			Flags.bIsExpanded = ItemDetails->IsExpandedByDefault();
+		}
 	}
 
 	bool FWorkspaceOutlinerTreeItem::IsValid() const
@@ -147,6 +151,11 @@ namespace UE::Workspace
 
 	FString FWorkspaceOutlinerTreeItem::GetDisplayString() const
 	{
+		if (const TSharedPtr<IWorkspaceOutlinerItemDetails> SharedFactory = FWorkspaceEditorModule::GetOutlinerItemDetails(MakeOutlinerDetailsId(Export)))
+		{
+			return SharedFactory->GetDisplayString(Export);
+		}
+		
 		return Export.GetIdentifier().ToString();
 	}
 
@@ -164,9 +173,9 @@ namespace UE::Workspace
 			const UPackage* Package = SharedFactory->GetPackage(Export);
 			return Package != nullptr ? Package->GetName() : FString();
 		}
-		else if (Export.GetParentIdentifier() == NAME_None && Export.GetAssetPath().IsValid())
+		else if (Export.GetParentIdentifier() == NAME_None && Export.GetFirstAssetPath().IsValid())
 		{
-			return Export.GetAssetPath().GetLongPackageName();
+			return Export.GetFirstAssetPath().GetLongPackageName();
 		}	
 		
 		return ISceneOutlinerTreeItem::GetPackageName();
diff --git a/Engine/Plugins/Experimental/Workspace/Source/WorkspaceEditor/Private/WorkspaceAssetReferenceItemDetails.h b/Engine/Plugins/Experimental/Workspace/Source/WorkspaceEditor/Private/WorkspaceAssetReferenceItemDetails.h
new file mode 100644
--- /dev/null
+++ b/Engine/Plugins/Experimental/Workspace/Source/WorkspaceEditor/Private/WorkspaceAssetReferenceItemDetails.h
@@ -0,0 +1,76 @@
+// Copyright Epic Games, Inc. All Rights Reserved.
+#pragma once
+
+#include "IWorkspaceOutlinerItemDetails.h"
+#include "WorkspaceEditorModule.h"
+#include "Styling/AppStyle.h"
+
+class UAnimNextRigVMAssetEditorData;
+enum class EAnimNextEditorDataNotifType : uint8;
+
+namespace UE::Workspace
+{
+class FWorkspaceAssetReferenceOutlinerItemDetails : public IWorkspaceOutlinerItemDetails
+{
+public:
+	virtual ~FWorkspaceAssetReferenceOutlinerItemDetails() override = default;
+
+	static TSharedPtr<IWorkspaceOutlinerItemDetails> GetInnerDetails(const FWorkspaceOutlinerItemExport& Export)
+	{
+		TSharedPtr<IWorkspaceOutlinerItemDetails> SharedFactory = nullptr;
+		const FWorkspaceOutlinerAssetReferenceItemData& Data = Export.GetData().Get<FWorkspaceOutlinerAssetReferenceItemData>();
+		if(Data.ReferredExport.HasData())
+		{
+			const FOutlinerItemDetailsId Id = Data.ReferredExport.GetData().GetScriptStruct()->GetFName();		
+			SharedFactory = FWorkspaceEditorModule::GetOutlinerItemDetails(Id);
+		}
+
+		return SharedFactory;
+	}
+
+	static const FWorkspaceOutlinerItemExport& GetInnerExport(const FWorkspaceOutlinerItemExport& Export)
+	{
+		FWorkspaceOutlinerItemExport InnerExport = Export;
+		check(Export.HasData());
+		const FWorkspaceOutlinerAssetReferenceItemData& Data = Export.GetData().Get<FWorkspaceOutlinerAssetReferenceItemData>();
+		return Data.ReferredExport;
+	}
+
+	virtual FString GetDisplayString(const FWorkspaceOutlinerItemExport& Export) const override
+	{
+		if (const TSharedPtr<IWorkspaceOutlinerItemDetails> SharedFactory = GetInnerDetails(Export))
+		{
+			return GetInnerDetails(Export)->GetDisplayString(GetInnerExport(Export));
+		}
+
+		return bShowFullIdentifierInOutliner ? FString::Printf(TEXT("%s - %s - %s"), *Export.GetIdentifier().ToString(), *(Export.GetData().GetScriptStruct() ? Export.GetData().GetScriptStruct()->GetFName() : NAME_None).ToString(), *Export.GetFirstAssetPath().ToString()) :
+			Export.GetFirstAssetPath().GetAssetName();
+	}
+	
+	virtual const FSlateBrush* GetItemIcon(const FWorkspaceOutlinerItemExport& Export) const override
+	{
+		const FSlateBrush* Icon = nullptr;
+		{
+			if (const TSharedPtr<IWorkspaceOutlinerItemDetails> SharedFactory = GetInnerDetails(Export))
+			{
+				const FWorkspaceOutlinerItemExport InnerExport = GetInnerExport(Export);
+				Icon = SharedFactory->GetItemIcon(InnerExport);
+			}
+		}
+				
+		return Icon ? Icon : FAppStyle::GetBrush(TEXT("GenericLink"));
+	}
+
+	virtual FSlateColor GetItemColor(const FWorkspaceOutlinerItemExport& Export) const override
+	{
+		const FWorkspaceOutlinerAssetReferenceItemData& Data = Export.GetData().Get<FWorkspaceOutlinerAssetReferenceItemData>();	
+		return Data.bRecursiveReference ? FLinearColor::Red : FSlateColor::UseSubduedForeground();
+	}	
+
+	virtual bool IsExpandedByDefault() const override
+	{
+		return false;
+	}	
+};
+
+}
diff --git a/Engine/Plugins/Experimental/Workspace/Source/WorkspaceEditor/Private/WorkspaceEditorModule.cpp b/Engine/Plugins/Experimental/Workspace/Source/WorkspaceEditor/Private/WorkspaceEditorModule.cpp
--- a/Engine/Plugins/Experimental/Workspace/Source/WorkspaceEditor/Private/WorkspaceEditorModule.cpp
+++ b/Engine/Plugins/Experimental/Workspace/Source/WorkspaceEditor/Private/WorkspaceEditorModule.cpp
@@ -18,6 +18,7 @@
 #include "EdGraph/EdGraphSchema.h"
 #include "Modules/ModuleManager.h"
 #include "Styling/AppStyle.h"
+#include "WorkspaceAssetReferenceItemDetails.h"
 
 #define LOCTEXT_NAMESPACE "WorkspaceEditorModule"
 
@@ -44,6 +45,9 @@ FWorkspaceEditorContext::FWorkspaceEditorContext(const TSharedRef<IWorkspaceEdit
 void FWorkspaceEditorModule::StartupModule()
 {
 	FWorkspaceAssetEditorCommands::Register();
+	
+	TSharedPtr<FWorkspaceAssetReferenceOutlinerItemDetails> ReferenceItemDetails = MakeShareable<FWorkspaceAssetReferenceOutlinerItemDetails>(new FWorkspaceAssetReferenceOutlinerItemDetails());
+	RegisterWorkspaceItemDetails(Workspace::FOutlinerItemDetailsId(FWorkspaceOutlinerAssetReferenceItemData::StaticStruct()->GetFName()), StaticCastSharedPtr<UE::Workspace::IWorkspaceOutlinerItemDetails>(ReferenceItemDetails));	
 }
 
 void FWorkspaceEditorModule::RegisterObjectDocumentType(const FTopLevelAssetPath& InClassPath, const FObjectDocumentArgs& InParams)
diff --git a/Engine/Plugins/Experimental/Workspace/Source/WorkspaceEditor/Private/WorkspaceEditorModule.h b/Engine/Plugins/Experimental/Workspace/Source/WorkspaceEditor/Private/WorkspaceEditorModule.h
--- a/Engine/Plugins/Experimental/Workspace/Source/WorkspaceEditor/Private/WorkspaceEditorModule.h
+++ b/Engine/Plugins/Experimental/Workspace/Source/WorkspaceEditor/Private/WorkspaceEditorModule.h
@@ -76,6 +76,7 @@ private:
 	friend class FWorkspaceOutlinerMode;
 	friend class FWorkspaceOutlinerSourceControlColumn;
 	friend class FWorkspaceOutlinerFileStateColumn;
+	friend class FWorkspaceAssetReferenceOutlinerItemDetails;
 	friend SWorkspaceTabWrapper;
 	friend FWorkspaceEditorContext;
 };
diff --git a/Engine/Plugins/Experimental/Workspace/Source/WorkspaceEditor/Public/IWorkspaceOutlinerItemDetails.h b/Engine/Plugins/Experimental/Workspace/Source/WorkspaceEditor/Public/IWorkspaceOutlinerItemDetails.h
--- a/Engine/Plugins/Experimental/Workspace/Source/WorkspaceEditor/Public/IWorkspaceOutlinerItemDetails.h
+++ b/Engine/Plugins/Experimental/Workspace/Source/WorkspaceEditor/Public/IWorkspaceOutlinerItemDetails.h
@@ -11,20 +11,27 @@ struct FSlateBrush;
 
 namespace UE::Workspace
 {
+	constexpr bool bShowFullIdentifierInOutliner = false;	
 
 typedef FName FOutlinerItemDetailsId;
 static FOutlinerItemDetailsId MakeOutlinerDetailsId(const FWorkspaceOutlinerItemExport& InExport)
 {
-    return InExport.GetData().IsValid() ? InExport.GetData().GetScriptStruct()->GetFName() : NAME_None;
+    return InExport.HasData() ? InExport.GetData().GetScriptStruct()->GetFName() : NAME_None;
 }
 
 class IWorkspaceOutlinerItemDetails : public TSharedFromThis<IWorkspaceOutlinerItemDetails>
 {
 public:
     virtual ~IWorkspaceOutlinerItemDetails() = default;
+	virtual FString GetDisplayString(const FWorkspaceOutlinerItemExport& Export) const
+	{			
+		return bShowFullIdentifierInOutliner ? FString::Printf(TEXT("%s - %s - %s"),
+		*Export.GetIdentifier().ToString(), *(Export.HasData() ? Export.GetData().GetScriptStruct()->GetFName() : NAME_None).ToString(), *Export.GetFirstAssetPath().ToString()) : Export.GetIdentifier().ToString();
+	}
+	
     virtual const FSlateBrush* GetItemIcon(const FWorkspaceOutlinerItemExport& Export) const { return nullptr; }
 	virtual FSlateColor GetItemColor(const FWorkspaceOutlinerItemExport& Export) const { return FSlateColor::UseForeground(); }
-    virtual void HandleDoubleClick(const FToolMenuContext& ToolMenuContext) const {}
+    virtual bool HandleDoubleClick(const FToolMenuContext& ToolMenuContext) const { return false; }
 	virtual bool CanDelete(const FWorkspaceOutlinerItemExport& Export) const { return true; }
 	virtual void Delete(TConstArrayView<FWorkspaceOutlinerItemExport> Exports) const {}
     virtual bool CanRename(const FWorkspaceOutlinerItemExport& Export) const { return false; }
@@ -32,6 +39,7 @@ public:
     virtual bool ValidateName(const FWorkspaceOutlinerItemExport& Export, const FText& InName, FText& OutErrorMessage) const { return false; }
     virtual UPackage* GetPackage(const FWorkspaceOutlinerItemExport& Export) const { return nullptr; }
 	virtual bool HandleSelected(const FToolMenuContext& ToolMenuContext) const { return false; }
+	virtual bool IsExpandedByDefault() const { return true; }
 };
 
 }
diff --git a/Engine/Plugins/Experimental/Workspace/Source/WorkspaceEditor/Public/WorkspaceAssetRegistryInfo.h b/Engine/Plugins/Experimental/Workspace/Source/WorkspaceEditor/Public/WorkspaceAssetRegistryInfo.h
--- a/Engine/Plugins/Experimental/Workspace/Source/WorkspaceEditor/Public/WorkspaceAssetRegistryInfo.h
+++ b/Engine/Plugins/Experimental/Workspace/Source/WorkspaceEditor/Public/WorkspaceAssetRegistryInfo.h
@@ -83,9 +83,12 @@ struct FWorkspaceOutlinerItemExport
 		Path.PathSegments.Add(InIdentifier);
 	}
 	
-	FWorkspaceOutlinerItemExport(const FName InIdentifier, const FWorkspaceOutlinerItemExport& InParent)
+	FWorkspaceOutlinerItemExport(const FName InIdentifier, const FWorkspaceOutlinerItemExport& InParent) : Path(InParent.Path.AppendSegment(InIdentifier))
+	{
+	}
+
+	FWorkspaceOutlinerItemExport(const FName InIdentifier, const FWorkspaceOutlinerItemExport& InParent, const TInstancedStruct<FWorkspaceOutlinerItemData>& InData) : Path(InParent.Path.AppendSegment(InIdentifier)),  Data(InData)
 	{
-		Path = InParent.Path.AppendSegment(InIdentifier);
 	}
 
 protected:
@@ -118,20 +121,100 @@ public:
 
 		return NAME_None;
 	}
-	
-	FSoftObjectPath GetAssetPath() const
+
+	template<typename AssetClass>
+	AssetClass* GetFirstAssetOfType() const
+	{
+		if(Path.PathSegments.Num() > 0)
+		{
+			for (int32 SegmentIndex = Path.PathSegments.Num() - 1; SegmentIndex >= 0; SegmentIndex--)
+			{
+				FSoftObjectPath ObjectPath = Path.PathSegments[SegmentIndex].ToString();
+				if (ObjectPath.IsValid() && ObjectPath.IsAsset())
+				{
+					if (UObject* Object = ObjectPath.TryLoad())
+					{
+						if (AssetClass* TypedObject = Cast<AssetClass>(Object))
+						{
+							return TypedObject;
+						}
+					}
+				}
+			}
+		}
+
+		return nullptr;
+	}
+
+	// Returns the first FSoftObjectPath found in segments, starting from the end. e.g. "SoftObjectPath" - "Foo" - "SoftObjectPathTwo" - "Bar" will return SoftObjectPathTwo
+	FSoftObjectPath GetFirstAssetPath() const
 	{
 		// Path needs atleast one segment to contain a (potentially) valid asset path
 		if(Path.PathSegments.Num() > 0)
 		{
-			return FSoftObjectPath(Path.PathSegments[0].ToString());	
+			for (int32 SegmentIndex = Path.PathSegments.Num() - 1; SegmentIndex >= 0; SegmentIndex--)
+			{
+				FSoftObjectPath ObjectPath = Path.PathSegments[SegmentIndex].ToString();
+				if (ObjectPath.IsValid() && ObjectPath.IsAsset())
+				{
+					return ObjectPath;	
+				}
+			}
+			
+			return FSoftObjectPath(Path.PathSegments[0].ToString());
 		}
+		
+		return FSoftObjectPath();
+	}
 
+	// Returns the first path segment as a FSoftObjectPath. e.g:
+	//  - "SoftObjectPath" - "Foo" - "SoftObjectPathTwo" - "Bar" will return SoftObjectPath	
+	//  - "Foo" - "SoftObjectPath" - "SoftObjectPathTwo" - "Bar" will return FSoftObjectPath()
+	FSoftObjectPath GetTopLevelAssetPath() const
+	{
+		// Path needs atleast one segment to contain a (potentially) valid asset path
+		if(Path.PathSegments.Num() > 0)
+		{
+			return FSoftObjectPath(Path.PathSegments[0].ToString());
+		}
+		
 		return FSoftObjectPath();
 	}
+
+	// Returns all valid FSoftObjectPaths found in path segments, starting from the end. e.g:
+	//  - "SoftObjectPath" - "Foo" - "SoftObjectPathTwo" - "Bar" will return "SoftObjectPathTwo", "SoftObjectPath"
+	//  - "Foo" - "SoftObjectPath" - "SoftObjectPathTwo" - "Bar" will also return "SoftObjectPathTwo", "SoftObjectPath"	
+	void GetAssetPaths(TArray<FSoftObjectPath>& OutAssetPaths) const
+	{
+		if(Path.PathSegments.Num() > 0)
+		{
+			for (int32 SegmentIndex = Path.PathSegments.Num() - 1; SegmentIndex >= 0; SegmentIndex--)
+			{
+				FSoftObjectPath ObjectPath = Path.PathSegments[SegmentIndex].ToString();
+				if (ObjectPath.IsValid() && ObjectPath.IsAsset())
+				{
+					OutAssetPaths.Add(ObjectPath);
+				}
+			}
+		}
+	}
+
+	FString GetFullPath() const
+	{
+		FString FullPath;
+		for (int32 SegmentIndex = 0; SegmentIndex < Path.PathSegments.Num(); SegmentIndex++)
+		{
+			FullPath.Append(Path.PathSegments[SegmentIndex].ToString());
+		}
+
+		return FullPath;
+	}
 	
 	// Remove identifier segment to retrieve parent path hash
 	uint32 GetParentHash() const { return GetTypeHash(Path.RemoveSegment()); }
+
+	// Returns whether or not Data has any instanced struct setup
+	bool HasData() const { return Data.IsValid(); }
 	
 	const TInstancedStruct<FWorkspaceOutlinerItemData>& GetData() const { return Data; }
 	TInstancedStruct<FWorkspaceOutlinerItemData>& GetData() { return Data; }	
@@ -142,6 +225,23 @@ public:
 	}
 };
 
+USTRUCT()
+struct WORKSPACEEDITOR_API FWorkspaceOutlinerAssetReferenceItemData : public FWorkspaceOutlinerItemData
+{
+	GENERATED_BODY()
+
+	FWorkspaceOutlinerAssetReferenceItemData() = default;
+
+	UPROPERTY()
+	FSoftObjectPath ReferredObjectPath;
+
+	UPROPERTY()
+	FWorkspaceOutlinerItemExport ReferredExport;
+
+	UPROPERTY()
+	bool bRecursiveReference = false;
+};
+
 namespace UE::Workspace
 {
 
