AnimNext: Fix new (i.e. unsaved) workspace entries not being considered when calling FWorkspaceEditorModule::OpenWorkspaceForObject

diff --git a/Engine/Plugins/Experimental/Workspace/Source/WorkspaceEditor/Private/WorkspaceEditorModule.cpp b/Engine/Plugins/Experimental/Workspace/Source/WorkspaceEditor/Private/WorkspaceEditorModule.cpp
--- a/Engine/Plugins/Experimental/Workspace/Source/WorkspaceEditor/Private/WorkspaceEditorModule.cpp
+++ b/Engine/Plugins/Experimental/Workspace/Source/WorkspaceEditor/Private/WorkspaceEditorModule.cpp
@@ -150,22 +150,33 @@ FObjectDocumentArgs FWorkspaceEditorModule::CreateGraphDocumentArgs(const FGraph
 
 bool FWorkspaceEditorModule::GetExportedAssetsForWorkspace(const FAssetData& InWorkspaceAsset, FWorkspaceAssetRegistryExports& OutExports)
 {
-	const FString ExternalObjectsPath = FExternalPackageHelper::GetExternalObjectsPath(InWorkspaceAsset.PackageName.ToString());
-	FARFilter Filter;
-	Filter.bRecursivePaths = true;
-	Filter.bIncludeOnlyOnDiskAssets = true;
-	Filter.PackagePaths.Add(*ExternalObjectsPath);
+	constexpr bool bLoadAsset = false;
+	if(UWorkspace* Workspace = Cast<UWorkspace>(InWorkspaceAsset.FastGetAsset(bLoadAsset)))
+	{
+		for(UWorkspaceAssetEntry* WorkspaceAssetEntry : Workspace->AssetEntries)
+		{
+			OutExports.Assets.Add({WorkspaceAssetEntry->Asset.ToSoftObjectPath()});
+		}
+	}
+	else
+	{
+		const FString ExternalObjectsPath = FExternalPackageHelper::GetExternalObjectsPath(InWorkspaceAsset.PackageName.ToString());
+		FARFilter Filter;
+		Filter.bRecursivePaths = true;
+		Filter.bIncludeOnlyOnDiskAssets = true;
+		Filter.PackagePaths.Add(*ExternalObjectsPath);
 
-	TArray<FAssetData> AssetDataEntries;
-	const IAssetRegistry& AssetRegistry = FAssetRegistryModule::GetRegistry();
-	AssetRegistry.GetAssets(Filter, AssetDataEntries);
+		TArray<FAssetData> AssetDataEntries;
+		IAssetRegistry& AssetRegistry = FAssetRegistryModule::GetRegistry();
+		AssetRegistry.GetAssets(Filter, AssetDataEntries);
 
-	for (const FAssetData& AssetDataEntry : AssetDataEntries)
-	{
-		FSoftObjectPath SoftObjectPath(AssetDataEntry.GetTagValueRef<FString>(UWorkspaceAssetEntry::ExportsAssetRegistryTag));
-		if(SoftObjectPath.IsValid())
+		for (const FAssetData& AssetDataEntry : AssetDataEntries)
 		{
-			OutExports.Assets.Add({SoftObjectPath});
+			FSoftObjectPath SoftObjectPath(AssetDataEntry.GetTagValueRef<FString>(UWorkspaceAssetEntry::ExportsAssetRegistryTag));
+			if(SoftObjectPath.IsValid())
+			{
+				OutExports.Assets.Add({SoftObjectPath});
+			}
 		}
 	}
 
