Rename 'Delete' to 'Remove' in the Workspace outliner to clarify functionality.

diff --git a/Engine/Plugins/Experimental/Workspace/Source/WorkspaceEditor/Private/Outliner/WorkspaceOutlinerMode.cpp b/Engine/Plugins/Experimental/Workspace/Source/WorkspaceEditor/Private/Outliner/WorkspaceOutlinerMode.cpp
--- a/Engine/Plugins/Experimental/Workspace/Source/WorkspaceEditor/Private/Outliner/WorkspaceOutlinerMode.cpp
+++ b/Engine/Plugins/Experimental/Workspace/Source/WorkspaceEditor/Private/Outliner/WorkspaceOutlinerMode.cpp
@@ -73,7 +73,7 @@ TSharedPtr<SWidget> FWorkspaceOutlinerMode::CreateContextMenu()
 				{
 					FToolMenuSection& CommonSection = InMenu->AddSection("Common", LOCTEXT("CommonSectionLabel", "Common"));
 					CommonSection.AddMenuEntryWithCommandList(FWorkspaceAssetEditorCommands::Get().Open, MenuContext->WeakCommandList.Pin());
-					CommonSection.AddMenuEntryWithCommandList(FGenericCommands::Get().Delete, MenuContext->WeakCommandList.Pin());
+					CommonSection.AddMenuEntryWithCommandList(FGenericCommands::Get().Delete, MenuContext->WeakCommandList.Pin(), LOCTEXT("RemoveLabel", "Remove"), LOCTEXT("RemoveTooltip", "Remove current selection"), FSlateIcon(FAppStyle::GetAppStyleSetName(), "Icons.Minus"));
 					CommonSection.AddMenuEntryWithCommandList(FGenericCommands::Get().Rename, MenuContext->WeakCommandList.Pin());
 					// TODO JDB more possible actions? (find in content browser?)
 
