AnimNext: Fix crash when clicking on an error/warning link

Only certain nodes implement SLATE_DECLARE_WIDGET_API, so we cant rely on StaticWidgetClass checks.

diff --git a/Engine/Plugins/Experimental/AnimNext/Source/AnimNextEditor/Private/AnimNextEditorModule.cpp b/Engine/Plugins/Experimental/AnimNext/Source/AnimNextEditor/Private/AnimNextEditorModule.cpp
--- a/Engine/Plugins/Experimental/AnimNext/Source/AnimNextEditor/Private/AnimNextEditorModule.cpp
+++ b/Engine/Plugins/Experimental/AnimNext/Source/AnimNextEditor/Private/AnimNextEditorModule.cpp
@@ -1141,7 +1141,7 @@ void FAnimNextEditorModule::RegisterWorkspaceDocumentTypes(Workspace::IWorkspace
 	{
 		if(UAnimNextEdGraphNode* EdGraphNode = Cast<UAnimNextEdGraphNode>(InObject))
 		{
-			check(InWidget->GetType() == SGraphEditor::StaticWidgetClass().GetWidgetType());
+			check(InWidget->GetType() == TEXT("SGraphEditor"));
 			TSharedPtr<SGraphEditor> GraphEditor = StaticCastSharedRef<SGraphEditor>(InWidget);
 			GraphEditor->JumpToNode(EdGraphNode, false);
 		}
diff --git a/Engine/Plugins/Experimental/Workspace/Source/WorkspaceEditor/Private/WorkspaceEditor.cpp b/Engine/Plugins/Experimental/Workspace/Source/WorkspaceEditor/Private/WorkspaceEditor.cpp
--- a/Engine/Plugins/Experimental/Workspace/Source/WorkspaceEditor/Private/WorkspaceEditor.cpp
+++ b/Engine/Plugins/Experimental/Workspace/Source/WorkspaceEditor/Private/WorkspaceEditor.cpp
@@ -348,10 +348,10 @@ void FWorkspaceEditor::OpenObjects(TConstArrayView<UObject*> InObjects)
 				TSharedRef<SWidget> TabContentWidget = WorkspaceTabWrapper->GetContent();
 
 				// If this is a built-in graph editor widget, supply the inner SGraphEditor
-				if(TabContentWidget->GetType() == SGraphDocument::StaticWidgetClass().GetWidgetType())
+				if(TabContentWidget->GetType() == TEXT("SGraphDocument"))
 				{
 					TabContentWidget = StaticCastSharedRef<SGraphDocument>(TabContentWidget)->GraphEditor.ToSharedRef();
-					check(TabContentWidget->GetType() == SGraphEditor::StaticWidgetClass().GetWidgetType());
+					check(TabContentWidget->GetType() == TEXT("SGraphEditor"));
 				}
 
 				DocumentSubObjectArgs->OnPostDocumentOpenedForSubObject.ExecuteIfBound(FWorkspaceEditorContext(SharedThis(this), Object), TabContentWidget, OriginalObject);
