Crash when docking the outliner tab in State Tree asset
#fix avoid recursion of FWorkspaceEditor::SetFocussedDocument. In this case was caused by SDockTab::SetContent inside of FStateTreeEditorModeToolkit::UpdateStateTreeOutliner which forcefully refreshes the tab's parent widget (which inadvertently refreshed the focussed document again)

diff --git a/Engine/Plugins/Experimental/Workspace/Source/WorkspaceEditor/Private/WorkspaceEditor.cpp b/Engine/Plugins/Experimental/Workspace/Source/WorkspaceEditor/Private/WorkspaceEditor.cpp
--- a/Engine/Plugins/Experimental/Workspace/Source/WorkspaceEditor/Private/WorkspaceEditor.cpp
+++ b/Engine/Plugins/Experimental/Workspace/Source/WorkspaceEditor/Private/WorkspaceEditor.cpp
@@ -452,6 +452,12 @@ void FWorkspaceEditor::SetGlobalSelection(FGlobalSelectionId SelectionId, FOnCle
 
 void FWorkspaceEditor::SetFocussedDocument(const TObjectPtr<UObject> InDocument)
 {
+	if (bSettingFocussedDocument)
+	{
+		return;
+	}	
+	TGuardValue<bool> SetValueGuard(bSettingFocussedDocument, true);
+	
 	FAssetEditorModeManager* ModeManager = static_cast<FAssetEditorModeManager*>(&GetEditorModeManager());	
 	ModeManager->GetSelectedObjects()->DeselectAll();
 
diff --git a/Engine/Plugins/Experimental/Workspace/Source/WorkspaceEditor/Private/WorkspaceEditor.h b/Engine/Plugins/Experimental/Workspace/Source/WorkspaceEditor/Private/WorkspaceEditor.h
--- a/Engine/Plugins/Experimental/Workspace/Source/WorkspaceEditor/Private/WorkspaceEditor.h
+++ b/Engine/Plugins/Experimental/Workspace/Source/WorkspaceEditor/Private/WorkspaceEditor.h
@@ -157,6 +157,7 @@ private:
 	bool bSavingWorkspace = false;
 	bool bSavingAssetEntries = false;
 	bool bClosingDown = false;
+	bool bSettingFocussedDocument = false;
 
 	TArray<TPair<FGlobalSelectionId, FOnClearGlobalSelection>> GlobalSelections;
 	FOnOutlinerSelectionChanged OnOutlinerSelectionChangedDelegate;
