Try to remove friction by only asking user to save transient workspaces if it contains more than 1 asset (which means the user actively chose to add an asset to it)

diff --git a/Engine/Plugins/Experimental/Workspace/Source/WorkspaceEditor/Private/WorkspaceEditor.cpp b/Engine/Plugins/Experimental/Workspace/Source/WorkspaceEditor/Private/WorkspaceEditor.cpp
--- a/Engine/Plugins/Experimental/Workspace/Source/WorkspaceEditor/Private/WorkspaceEditor.cpp
+++ b/Engine/Plugins/Experimental/Workspace/Source/WorkspaceEditor/Private/WorkspaceEditor.cpp
@@ -789,7 +789,8 @@ bool FWorkspaceEditor::OnRequestClose(EAssetEditorCloseReason InCloseReason)
 	auto RequiresSave = [this]()
 	{
 		const UPackage* Package = Workspace->GetOutermost();
-		return Package->GetPathName().StartsWith(TEXT("/Temp/Untitled"));
+		// Ask the user to save a transient workspace containing more than 1 asset
+		return Package->GetPathName().StartsWith(TEXT("/Temp/Untitled")) && Workspace->AssetEntries.Num() > 1;
 	};
 
 	// Give the user opportunity to save temp workspaces
