Add new test plugin inside EngineTest.
Automate QMetry Test Case:  UE-TC-10715

diff --git a/Engine/Plugins/Experimental/Workspace/Source/WorkspaceEditor/Private/Workspace.h b/Engine/Plugins/Experimental/Workspace/Source/WorkspaceEditor/Private/Workspace.h
--- a/Engine/Plugins/Experimental/Workspace/Source/WorkspaceEditor/Private/Workspace.h
+++ b/Engine/Plugins/Experimental/Workspace/Source/WorkspaceEditor/Private/Workspace.h
@@ -60,7 +60,7 @@ struct FWorkspaceAssetRegistryExports
 	TArray<FWorkspaceAssetRegistryExportEntry> Assets;
 };
 
-UCLASS()
+UCLASS(MinimalAPI)
 class UWorkspace : public UObject
 {
 	GENERATED_BODY()
@@ -156,4 +156,4 @@ class UWorkspace : public UObject
 	UE::Workspace::FOnWorkspaceModified ModifiedDelegate;
 
 	bool bSuspendNotifications = false;
-};
\ No newline at end of file
+};
