Workspace document restructure
- Added WorkspaceExport as part of document definition (this allows for more contextual information during editing)
- Workspace Asset document summoner now derives from FDocumentTabFactory mirror FDocumentTabFactoryForObjects<UObject> behaviour but for FWorkspaceDocument payload(s) instead
- Some miscellaneous export fixes (exports for null objects)
- Added API for Opening and Selection WorkspaceItemExports directly
- DocumentStates now also capture the export for correct state restoration
- Workspace tab breadcrumbs are now generated from WorkspaceItemExport chain if available (otherwise falls back to old behaviour)

diff --git a/Engine/Plugins/Experimental/AnimNext/Source/AnimNextEditor/Internal/Common/AnimNextAssetItemDetails.cpp b/Engine/Plugins/Experimental/AnimNext/Source/AnimNextEditor/Internal/Common/AnimNextAssetItemDetails.cpp
--- a/Engine/Plugins/Experimental/AnimNext/Source/AnimNextEditor/Internal/Common/AnimNextAssetItemDetails.cpp
+++ b/Engine/Plugins/Experimental/AnimNext/Source/AnimNextEditor/Internal/Common/AnimNextAssetItemDetails.cpp
@@ -83,7 +83,7 @@ void FAnimNextAssetItemDetails::RegisterToolMenuExtensions()
 			return;
 		}
 
-		TInstancedStruct<FWorkspaceOutlinerItemData>& Data = WorkspaceItemContext->SelectedExports[0].GetData();
+		TInstancedStruct<FWorkspaceOutlinerItemData>& Data = WorkspaceItemContext->SelectedExports[0].GetResolvedExport().GetData();
 		if (!Data.IsValid() || !Data.GetScriptStruct()->IsChildOf(FAnimNextRigVMAssetOutlinerData::StaticStruct()))
 		{
 			return;
diff --git a/Engine/Plugins/Experimental/AnimNext/Source/AnimNextEditor/Private/AnimNextEditorModule.cpp b/Engine/Plugins/Experimental/AnimNext/Source/AnimNextEditor/Private/AnimNextEditorModule.cpp
--- a/Engine/Plugins/Experimental/AnimNext/Source/AnimNextEditor/Private/AnimNextEditorModule.cpp
+++ b/Engine/Plugins/Experimental/AnimNext/Source/AnimNextEditor/Private/AnimNextEditorModule.cpp
@@ -790,7 +790,7 @@ void FAnimNextEditorModule::RegisterWorkspaceDocumentTypes(Workspace::IWorkspace
 			return;
 		}
 
-		URigVMEdGraph* RigVMEdGraph = Cast<URigVMEdGraph>(InContext.Object);
+		URigVMEdGraph* RigVMEdGraph = Cast<URigVMEdGraph>(InContext.Document.Object);
 		if (RigVMEdGraph == nullptr)
 		{
 			return;
@@ -819,7 +819,7 @@ void FAnimNextEditorModule::RegisterWorkspaceDocumentTypes(Workspace::IWorkspace
 			return;
 		}
 
-		URigVMEdGraph* RigVMEdGraph = Cast<URigVMEdGraph>(InContext.Object);
+		URigVMEdGraph* RigVMEdGraph = Cast<URigVMEdGraph>(InContext.Document.Object);
 		if (RigVMEdGraph == nullptr)
 		{
 			return;
@@ -865,7 +865,7 @@ void FAnimNextEditorModule::RegisterWorkspaceDocumentTypes(Workspace::IWorkspace
 				return;
 			}
 
-			URigVMEdGraph* RigVMEdGraph = Cast<URigVMEdGraph>(InContext.Object);
+			URigVMEdGraph* RigVMEdGraph = Cast<URigVMEdGraph>(InContext.Document.Object);
 			if (RigVMEdGraph == nullptr)
 			{
 				return;
@@ -906,7 +906,7 @@ void FAnimNextEditorModule::RegisterWorkspaceDocumentTypes(Workspace::IWorkspace
 				return;
 			}
 
-			URigVMEdGraph* RigVMEdGraph = Cast<URigVMEdGraph>(InContext.Object);
+			URigVMEdGraph* RigVMEdGraph = Cast<URigVMEdGraph>(InContext.Document.Object);
 			if (RigVMEdGraph == nullptr)
 			{
 				return;
@@ -928,7 +928,7 @@ void FAnimNextEditorModule::RegisterWorkspaceDocumentTypes(Workspace::IWorkspace
 		});
 	GraphArgs.OnGraphSelectionChanged = Workspace::FOnGraphSelectionChanged::CreateLambda([this](const Workspace::FWorkspaceEditorContext& InContext, const FGraphPanelSelectionSet& NewSelection)
 	{
-		URigVMEdGraph* RigVMEdGraph = Cast<URigVMEdGraph>(InContext.Object);
+		URigVMEdGraph* RigVMEdGraph = Cast<URigVMEdGraph>(InContext.Document.Object);
 		if (RigVMEdGraph == nullptr)
 		{
 			return;
@@ -1003,7 +1003,7 @@ void FAnimNextEditorModule::RegisterWorkspaceDocumentTypes(Workspace::IWorkspace
 	{
 		TWeakPtr<Workspace::IWorkspaceEditor> WeakWorkspaceEditor = InContext.WorkspaceEditor;
 
-		if (UAnimNextEdGraph* EdGraph = Cast<UAnimNextEdGraph>(InContext.Object))
+		if (UAnimNextEdGraph* EdGraph = Cast<UAnimNextEdGraph>(InContext.Document.Object))
 		{
 			UAnimNextRigVMAssetEditorData* EditorData = EdGraph->GetTypedOuter<UAnimNextRigVMAssetEditorData>();
 			if(EditorData)
@@ -1074,7 +1074,7 @@ void FAnimNextEditorModule::RegisterWorkspaceDocumentTypes(Workspace::IWorkspace
 
 	GraphDocumentArgs.OnGetDocumentBreadcrumbTrail = Workspace::FOnGetDocumentBreadcrumbTrail::CreateLambda([](const Workspace::FWorkspaceEditorContext& InContext, TArray<TSharedPtr<Workspace::FWorkspaceBreadcrumb>>& OutBreadcrumbs)
 	{
-		if (URigVMEdGraph* EdGraph = Cast<URigVMEdGraph>(InContext.Object))
+		if (URigVMEdGraph* EdGraph = Cast<URigVMEdGraph>(InContext.Document.Object))
 		{
 			if (UAnimNextRigVMAssetEditorData* EditorData = EdGraph->GetTypedOuter<UAnimNextRigVMAssetEditorData>())
 			{
@@ -1140,11 +1140,11 @@ void FAnimNextEditorModule::RegisterWorkspaceDocumentTypes(Workspace::IWorkspace
 							}
 						);
 						GraphCrumb->OnClicked = Workspace::FWorkspaceBreadcrumb::FOnBreadcrumbClicked::CreateLambda(
-							[WeakEditorData, WeakEdGraph, WeakWorkspaceEditor]
+							[WeakEditorData, WeakEdGraph, WeakWorkspaceEditor, Export = InContext.Document.Export]
 							{
 								if (const TSharedPtr<Workspace::IWorkspaceEditor> SharedWorkspaceEditor = WeakWorkspaceEditor.Pin())
 								{
-									SharedWorkspaceEditor->OpenObjects({ WeakEdGraph.Get() });
+									SharedWorkspaceEditor->OpenExports({ Export });
 								}
 							}
 						);
@@ -1170,11 +1170,11 @@ void FAnimNextEditorModule::RegisterWorkspaceDocumentTypes(Workspace::IWorkspace
 					TWeakPtr<Workspace::IWorkspaceEditor> WeakWorkspaceEditor = InContext.WorkspaceEditor;
 					OuterGraphCrumb->OnGetLabel = Workspace::FWorkspaceBreadcrumb::FOnGetBreadcrumbLabel::CreateLambda([AssetName = OuterAsset->GetFName()]{ return FText::FromName(AssetName); });
 					OuterGraphCrumb->OnClicked = Workspace::FWorkspaceBreadcrumb::FOnBreadcrumbClicked::CreateLambda(
-						[WeakOuterAsset, WeakWorkspaceEditor]
+						[WeakOuterAsset, WeakWorkspaceEditor, Export = InContext.Document.Export]
 						{
 							if (const TSharedPtr<Workspace::IWorkspaceEditor> SharedWorkspaceEditor = WeakWorkspaceEditor.Pin())
 							{
-								SharedWorkspaceEditor->OpenObjects({WeakOuterAsset.Get()});
+								SharedWorkspaceEditor->OpenExports({Export});
 							}
 						}
 					);
diff --git a/Engine/Plugins/Experimental/AnimNext/Source/AnimNextEditor/Private/Common/AnimNextCollapseNodeItemDetails.cpp b/Engine/Plugins/Experimental/AnimNext/Source/AnimNextEditor/Private/Common/AnimNextCollapseNodeItemDetails.cpp
--- a/Engine/Plugins/Experimental/AnimNext/Source/AnimNextEditor/Private/Common/AnimNextCollapseNodeItemDetails.cpp
+++ b/Engine/Plugins/Experimental/AnimNext/Source/AnimNextEditor/Private/Common/AnimNextCollapseNodeItemDetails.cpp
@@ -28,7 +28,7 @@ bool FAnimNextCollapseNodeItemDetails::HandleDoubleClick(const FToolMenuContext&
 	{
 		if (const TSharedPtr<UE::Workspace::IWorkspaceEditor> WorkspaceEditor = StaticCastSharedPtr<UE::Workspace::IWorkspaceEditor>(AssetEditorContext->Toolkit.Pin()))
 		{
-			const TInstancedStruct<FWorkspaceOutlinerItemData>& Data = WorkspaceItemContext->SelectedExports[0].GetData();
+			const TInstancedStruct<FWorkspaceOutlinerItemData>& Data = WorkspaceItemContext->SelectedExports[0].GetResolvedExport().GetData();
 			if (Data.IsValid() && Data.GetScriptStruct() == FAnimNextCollapseGraphOutlinerData::StaticStruct())
 			{
 				const FAnimNextCollapseGraphOutlinerData& CollapseGraphData = Data.Get<FAnimNextCollapseGraphOutlinerData>();
diff --git a/Engine/Plugins/Experimental/AnimNext/Source/AnimNextEditor/Private/Common/AnimNextFunctionItemDetails.cpp b/Engine/Plugins/Experimental/AnimNext/Source/AnimNextEditor/Private/Common/AnimNextFunctionItemDetails.cpp
--- a/Engine/Plugins/Experimental/AnimNext/Source/AnimNextEditor/Private/Common/AnimNextFunctionItemDetails.cpp
+++ b/Engine/Plugins/Experimental/AnimNext/Source/AnimNextEditor/Private/Common/AnimNextFunctionItemDetails.cpp
@@ -28,7 +28,7 @@ bool FAnimNextFunctionItemDetails::HandleDoubleClick(const FToolMenuContext& Too
 	{
 		if (const TSharedPtr<UE::Workspace::IWorkspaceEditor> WorkspaceEditor = StaticCastSharedPtr<UE::Workspace::IWorkspaceEditor>(AssetEditorContext->Toolkit.Pin()))
 		{
-			const TInstancedStruct<FWorkspaceOutlinerItemData>& Data = WorkspaceItemContext->SelectedExports[0].GetData();
+			const TInstancedStruct<FWorkspaceOutlinerItemData>& Data = WorkspaceItemContext->SelectedExports[0].GetResolvedExport().GetData();
 			if (Data.IsValid() && Data.GetScriptStruct() == FAnimNextGraphFunctionOutlinerData::StaticStruct())
 			{
 				const FAnimNextGraphFunctionOutlinerData& FunctionData = Data.Get<FAnimNextGraphFunctionOutlinerData>();
@@ -51,7 +51,7 @@ bool FAnimNextFunctionItemDetails::HandleSelected(const FToolMenuContext& ToolMe
 	{
 		if (const TSharedPtr<UE::Workspace::IWorkspaceEditor> WorkspaceEditor = StaticCastSharedPtr<UE::Workspace::IWorkspaceEditor>(AssetEditorContext->Toolkit.Pin()))
 		{
-			const TInstancedStruct<FWorkspaceOutlinerItemData>& Data = WorkspaceItemContext->SelectedExports[0].GetData();
+			const TInstancedStruct<FWorkspaceOutlinerItemData>& Data = WorkspaceItemContext->SelectedExports[0].GetResolvedExport().GetData();
 			if (Data.IsValid() && Data.GetScriptStruct() == FAnimNextGraphFunctionOutlinerData::StaticStruct())
 			{
 				const FAnimNextGraphFunctionOutlinerData& FunctionData = Data.Get<FAnimNextGraphFunctionOutlinerData>();
diff --git a/Engine/Plugins/Experimental/AnimNext/Source/AnimNextEditor/Private/Common/AnimNextGraphItemDetails.cpp b/Engine/Plugins/Experimental/AnimNext/Source/AnimNextEditor/Private/Common/AnimNextGraphItemDetails.cpp
--- a/Engine/Plugins/Experimental/AnimNext/Source/AnimNextEditor/Private/Common/AnimNextGraphItemDetails.cpp
+++ b/Engine/Plugins/Experimental/AnimNext/Source/AnimNextEditor/Private/Common/AnimNextGraphItemDetails.cpp
@@ -28,7 +28,7 @@ bool FAnimNextGraphItemDetails::HandleDoubleClick(const FToolMenuContext& ToolMe
 	{
 		if(const TSharedPtr<UE::Workspace::IWorkspaceEditor> WorkspaceEditor = StaticCastSharedPtr<UE::Workspace::IWorkspaceEditor>(AssetEditorContext->Toolkit.Pin()))
 		{
-			const TInstancedStruct<FWorkspaceOutlinerItemData>& Data = WorkspaceItemContext->SelectedExports[0].GetData();
+			const TInstancedStruct<FWorkspaceOutlinerItemData>& Data = WorkspaceItemContext->SelectedExports[0].GetResolvedExport().GetData();
 			if (Data.IsValid() && Data.GetScriptStruct() == FAnimNextGraphOutlinerData::StaticStruct())
 			{
 				const FAnimNextGraphOutlinerData& GraphData = Data.Get<FAnimNextGraphOutlinerData>();							
diff --git a/Engine/Plugins/Experimental/AnimNext/Source/AnimNextUncookedOnly/Private/UncookedOnlyUtils.cpp b/Engine/Plugins/Experimental/AnimNext/Source/AnimNextUncookedOnly/Private/UncookedOnlyUtils.cpp
--- a/Engine/Plugins/Experimental/AnimNext/Source/AnimNextUncookedOnly/Private/UncookedOnlyUtils.cpp
+++ b/Engine/Plugins/Experimental/AnimNext/Source/AnimNextUncookedOnly/Private/UncookedOnlyUtils.cpp
@@ -1005,10 +1005,15 @@ void ProcessPinAssetReferences(const URigVMPin* InPin, FWorkspaceOutlinerItemExp
 			const FSoftObjectPath ObjectPath(InPin->GetDefaultValue());
 			if (ObjectPath.IsValid())
 			{
-				FWorkspaceOutlinerItemExport& ParentExport = ParentExportIndex == INDEX_NONE ? RootExport : OutExports.Exports[ParentExportIndex];
-				FWorkspaceOutlinerItemExport& ReferenceExport = OutExports.Exports.Add_GetRef(FWorkspaceOutlinerItemExport(FName(ObjectPath.ToString()), ParentExport));
-				ReferenceExport.GetData().InitializeAs<FWorkspaceOutlinerAssetReferenceItemData>();
-				ReferenceExport.GetData().GetMutable<FWorkspaceOutlinerAssetReferenceItemData>().ReferredObjectPath = ObjectPath;
+				// Only add export if object is loaded, or the path actually points to an asset
+				FAssetData ReferenceAssetData;
+				if (ObjectPath.ResolveObject() || IAssetRegistry::GetChecked().TryGetAssetByObjectPath(ObjectPath, ReferenceAssetData) == UE::AssetRegistry::EExists::Exists)
+				{
+					FWorkspaceOutlinerItemExport& ParentExport = ParentExportIndex == INDEX_NONE ? RootExport : OutExports.Exports[ParentExportIndex];
+					FWorkspaceOutlinerItemExport& ReferenceExport = OutExports.Exports.Add_GetRef(FWorkspaceOutlinerItemExport(FName(ObjectPath.ToString()), ParentExport));
+					ReferenceExport.GetData().InitializeAs<FWorkspaceOutlinerAssetReferenceItemData>();
+					ReferenceExport.GetData().GetMutable<FWorkspaceOutlinerAssetReferenceItemData>().ReferredObjectPath = ObjectPath;
+				}
 			}
 		}
 
diff --git a/Engine/Plugins/Experimental/AnimNextStateTree/Source/AnimNextStateTreeEditor/Private/AnimNextStateTreeEditorModule.cpp b/Engine/Plugins/Experimental/AnimNextStateTree/Source/AnimNextStateTreeEditor/Private/AnimNextStateTreeEditorModule.cpp
--- a/Engine/Plugins/Experimental/AnimNextStateTree/Source/AnimNextStateTreeEditor/Private/AnimNextStateTreeEditorModule.cpp
+++ b/Engine/Plugins/Experimental/AnimNextStateTree/Source/AnimNextStateTreeEditor/Private/AnimNextStateTreeEditorModule.cpp
@@ -73,7 +73,7 @@ void FAnimNextStateTreeEditorModule::StartupModule()
 	Workspace::FObjectDocumentArgs StateTreeDocumentArgs(
 		Workspace::FOnMakeDocumentWidget::CreateLambda([](const Workspace::FWorkspaceEditorContext& InContext)-> TSharedRef<SWidget>
 		{
-			UAnimNextStateTree* AnimNextStateTree = CastChecked<UAnimNextStateTree>(InContext.Object);
+			UAnimNextStateTree* AnimNextStateTree = InContext.Document.GetTypedObject<UAnimNextStateTree>();
 			UStateTree* StateTree = AnimNextStateTree->StateTree;
 
 			TWeakPtr<Workspace::IWorkspaceEditor> WeakWorkspaceEditor = InContext.WorkspaceEditor;
@@ -133,7 +133,7 @@ void FAnimNextStateTreeEditorModule::StartupModule()
 	
 	StateTreeDocumentArgs.OnGetTabName = Workspace::FOnGetTabName::CreateLambda([](const Workspace::FWorkspaceEditorContext& InContext)
 	{
-		UAnimNextStateTree* AnimNextStateTree = CastChecked<UAnimNextStateTree>(InContext.Object);
+		UAnimNextStateTree* AnimNextStateTree = InContext.Document.GetTypedObject<UAnimNextStateTree>();
 		return FText::FromName(AnimNextStateTree->GetFName());
 	});
 
@@ -141,7 +141,7 @@ void FAnimNextStateTreeEditorModule::StartupModule()
 
 	StateTreeDocumentArgs.OnGetDocumentBreadcrumbTrail = Workspace::FOnGetDocumentBreadcrumbTrail::CreateLambda([](const Workspace::FWorkspaceEditorContext& InContext, TArray<TSharedPtr<Workspace::FWorkspaceBreadcrumb>>& OutBreadcrumbs)
 	{
-		if (const UAnimNextStateTree* AnimNextStateTree = CastChecked<UAnimNextStateTree>(InContext.Object))
+		if (UAnimNextStateTree* AnimNextStateTree = InContext.Document.GetTypedObject<UAnimNextStateTree>())
 		{
 			const TSharedPtr<Workspace::FWorkspaceBreadcrumb>& GraphCrumb = OutBreadcrumbs.Add_GetRef(MakeShared<Workspace::FWorkspaceBreadcrumb>());
 			GraphCrumb->OnGetLabel = Workspace::FWorkspaceBreadcrumb::FOnGetBreadcrumbLabel::CreateLambda([StateTreeName = AnimNextStateTree->GetFName()]{ return FText::FromName(StateTreeName); });
@@ -152,6 +152,18 @@ void FAnimNextStateTreeEditorModule::StartupModule()
 					return AnimNextStateTree->GetPackage()->IsDirty();
 				}
 			);
+
+			TWeakPtr<Workspace::IWorkspaceEditor> WeakWorkspaceEditor = InContext.WorkspaceEditor;
+			TWeakObjectPtr<UAnimNextStateTree> WeakStateTree = AnimNextStateTree;
+			GraphCrumb->OnClicked = Workspace::FWorkspaceBreadcrumb::FOnBreadcrumbClicked::CreateLambda(
+			[WeakWorkspaceEditor, Export = InContext.Document.Export]
+			{
+				if (const TSharedPtr<Workspace::IWorkspaceEditor> SharedWorkspaceEditor = WeakWorkspaceEditor.Pin())
+				{
+					SharedWorkspaceEditor->OpenExports({ Export });
+				}
+			}
+		);
 		}
 	});
 
@@ -257,13 +269,13 @@ void FAnimNextStateTreeEditorModule::StartupModule()
 					const uint32 NumSelectedExports = WorkspaceItemContext->SelectedExports.Num();				
 					if (NumSelectedExports == 1)
 					{
-						const FWorkspaceOutlinerItemExport& SelectedExport = WorkspaceItemContext->SelectedExports[0];
+						const FWorkspaceOutlinerItemExport& SelectedExport = WorkspaceItemContext->SelectedExports[0].GetResolvedExport();
 
-						if(WorkspaceItemContext->SelectedExports[0].GetData().GetScriptStruct() == FAnimNextStateTreeOutlinerData::StaticStruct())
+						if(SelectedExport.GetData().GetScriptStruct() == FAnimNextStateTreeOutlinerData::StaticStruct())
 						{
 							if (UAnimNextStateTree* LoadedStateTree = SelectedExport.GetFirstAssetOfType<UAnimNextStateTree>())
 							{
-								WorkspaceEditor->OpenAssets({LoadedStateTree });
+								WorkspaceEditor->OpenExports({ WorkspaceItemContext->SelectedExports[0] });
 
 								return true;
 							}
@@ -273,7 +285,7 @@ void FAnimNextStateTreeEditorModule::StartupModule()
 							if (UAnimNextStateTree* SelectionStateTree = SelectedExport.GetFirstAssetOfType<UAnimNextStateTree>())
 							{
 								// First open editor for state-tree
-								WorkspaceEditor->OpenAssets({SelectionStateTree });
+								WorkspaceEditor->OpenExports({ WorkspaceItemContext->SelectedExports[0] });
 
 								// Then select the double-clicked state
 								SelectStateExports(WorkspaceItemContext->SelectedExports);
@@ -300,9 +312,9 @@ void FAnimNextStateTreeEditorModule::StartupModule()
 					if (NumSelectedExports >= 1)
 					{
 						// Selecting AnimStateTree asset itself
-						if (NumSelectedExports == 1 && WorkspaceItemContext->SelectedExports[0].GetData().GetScriptStruct() == FAnimNextStateTreeOutlinerData::StaticStruct())
+						if (NumSelectedExports == 1 && WorkspaceItemContext->SelectedExports[0].GetResolvedExport().GetData().GetScriptStruct() == FAnimNextStateTreeOutlinerData::StaticStruct())
 						{
-							const FWorkspaceOutlinerItemExport& SelectedExport = WorkspaceItemContext->SelectedExports[0];
+							const FWorkspaceOutlinerItemExport& SelectedExport = WorkspaceItemContext->SelectedExports[0].GetResolvedExport();
 							if (UAnimNextStateTree* LoadedStateTree = Cast<UAnimNextStateTree>(SelectedExport.GetFirstAssetPath().ResolveObject()))
 							{
 								WorkspaceEditor->SetDetailsObjects({ LoadedStateTree->StateTree->EditorData } );
diff --git a/Engine/Plugins/Experimental/AnimNextStateTree/Source/AnimNextStateTreeUncookedOnly/Private/AnimNextStateTreeWorkspaceAssetUserData.cpp b/Engine/Plugins/Experimental/AnimNextStateTree/Source/AnimNextStateTreeUncookedOnly/Private/AnimNextStateTreeWorkspaceAssetUserData.cpp
--- a/Engine/Plugins/Experimental/AnimNextStateTree/Source/AnimNextStateTreeUncookedOnly/Private/AnimNextStateTreeWorkspaceAssetUserData.cpp
+++ b/Engine/Plugins/Experimental/AnimNextStateTree/Source/AnimNextStateTreeUncookedOnly/Private/AnimNextStateTreeWorkspaceAssetUserData.cpp
@@ -64,11 +64,16 @@ void UAnimNextStateTreeWorkspaceAssetUserData::GetAssetRegistryTags(FAssetRegist
 
 					for (const UObject* ReferredObject : ReferencedObjects)
 					{
-						FSoftObjectPath Path = ReferredObject;
-						FWorkspaceOutlinerItemExport& GraphReference = OutlinerExports.Exports.Add_GetRef(FWorkspaceOutlinerItemExport(FName(Path.ToString()), StateExport));
-						GraphReference.GetData().InitializeAsScriptStruct(FWorkspaceOutlinerAssetReferenceItemData::StaticStruct());
-						FWorkspaceOutlinerAssetReferenceItemData& ReferenceData = GraphReference.GetData().GetMutable<FWorkspaceOutlinerAssetReferenceItemData>();
-						ReferenceData.ReferredObjectPath = Path;
+						// GetObjectReferences will return null for unset object reference(s)
+						if (ReferredObject != nullptr)
+						{
+							FSoftObjectPath Path = ReferredObject;
+							FWorkspaceOutlinerItemExport& GraphReference = OutlinerExports.Exports.Add_GetRef(FWorkspaceOutlinerItemExport(FName(Path.ToString()), StateExport));
+							GraphReference.GetData().InitializeAsScriptStruct(FWorkspaceOutlinerAssetReferenceItemData::StaticStruct());
+							FWorkspaceOutlinerAssetReferenceItemData& ReferenceData = GraphReference.GetData().GetMutable<FWorkspaceOutlinerAssetReferenceItemData>();
+							ReferenceData.ReferredObjectPath = Path;
+						}
+						
 					}
 				}
 				
diff --git a/Engine/Plugins/Experimental/Workspace/Source/WorkspaceEditor/Private/AssetDocumentSummoner.cpp b/Engine/Plugins/Experimental/Workspace/Source/WorkspaceEditor/Private/AssetDocumentSummoner.cpp
--- a/Engine/Plugins/Experimental/Workspace/Source/WorkspaceEditor/Private/AssetDocumentSummoner.cpp
+++ b/Engine/Plugins/Experimental/Workspace/Source/WorkspaceEditor/Private/AssetDocumentSummoner.cpp
@@ -11,6 +11,7 @@
 #include "Widgets/Layout/SSpacer.h"
 #include "WorkspaceDocumentState.h"
 #include "SWorkspaceTabWrapper.h"
+#include "WorkspaceTabPayload.h"
 #include "Modules/ModuleManager.h"
 
 #define LOCTEXT_NAMESPACE "AssetDocumentSummoner"
@@ -18,8 +19,10 @@
 namespace UE::Workspace
 {
 
+const FName FTabPayload_WorkspaceDocument::DocumentPayloadName(TEXT("WorkspaceDocumentPayload"));
+
 FAssetDocumentSummoner::FAssetDocumentSummoner(FName InIdentifier, TSharedPtr<FWorkspaceEditor> InHostingApp, bool bInAllowUnsupportedClasses /*= false*/)
-	: FDocumentTabFactoryForObjects<UObject>(InIdentifier, InHostingApp)
+	: FDocumentTabFactory(InIdentifier, InHostingApp)
 	, HostingAppPtr(InHostingApp)
 	, bAllowUnsupportedClasses(bInAllowUnsupportedClasses)
 {
@@ -51,7 +54,7 @@ void FAssetDocumentSummoner::SaveState(TSharedPtr<SDockTab> Tab, TSharedPtr<FTab
 {
 	if(TSharedPtr<FWorkspaceEditor> WorkspaceEditor = HostingAppPtr.Pin())
 	{
-		UObject* Object = Payload->IsValid() ? FTabPayload_UObject::CastChecked<UObject>(Payload) : nullptr;
+		UObject* Object = Payload->IsValid() ? FTabPayload_WorkspaceDocument::CastChecked<UObject>(Payload.ToSharedRef()) : nullptr;
 		if(Object)
 		{
 			TSharedRef<SWorkspaceTabWrapper> TabWrapper = StaticCastSharedRef<SWorkspaceTabWrapper>(Tab->GetContent());
@@ -59,19 +62,20 @@ void FAssetDocumentSummoner::SaveState(TSharedPtr<SDockTab> Tab, TSharedPtr<FTab
 			const FObjectDocumentArgs* DocumentArgs = WorkspaceEditorModule.FindObjectDocumentType(Object);
 			if(DocumentArgs && DocumentArgs->OnGetDocumentState.IsBound())
 			{
-				WorkspaceEditor->RecordDocumentState(DocumentArgs->OnGetDocumentState.Execute(FWorkspaceEditorContext(WorkspaceEditor.ToSharedRef(), Object), TabWrapper->GetContent()));
+				WorkspaceEditor->RecordDocumentState(DocumentArgs->OnGetDocumentState.Execute(FWorkspaceEditorContext(WorkspaceEditor.ToSharedRef(), { FTabPayload_WorkspaceDocument::GetExport(Payload.ToSharedRef()), Object }), TabWrapper->GetContent()));
 			}
 			else
 			{
-				WorkspaceEditor->RecordDocumentState(TInstancedStruct<FWorkspaceDocumentState>::Make(Object));
+				WorkspaceEditor->RecordDocumentState(TInstancedStruct<FWorkspaceDocumentState>::Make(Object, FTabPayload_WorkspaceDocument::GetExport(Payload.ToSharedRef())));
 			}
 		}
 	}
 }
 
-TAttribute<FText> FAssetDocumentSummoner::ConstructTabNameForObject(UObject* DocumentID) const
+TAttribute<FText> FAssetDocumentSummoner::ConstructTabName(const FWorkflowTabSpawnInfo& Info) const
 {
 	TSharedPtr<FWorkspaceEditor> WorkspaceEditor = HostingAppPtr.Pin();
+	UObject* DocumentID = FTabPayload_WorkspaceDocument::CastChecked<UObject>(Info.Payload.ToSharedRef());
 	if(!WorkspaceEditor.IsValid() || DocumentID == nullptr)
 	{
 		return LOCTEXT("NoneObjectName", "None");
@@ -82,7 +86,7 @@ TAttribute<FText> FAssetDocumentSummoner::ConstructTabNameForObject(UObject* Doc
 	{
 		if(DocumentArgs->OnGetTabName.IsBound())
 		{
-			return DocumentArgs->OnGetTabName.Execute(FWorkspaceEditorContext(WorkspaceEditor.ToSharedRef(), DocumentID));
+			return DocumentArgs->OnGetTabName.Execute(FWorkspaceEditorContext(WorkspaceEditor.ToSharedRef(), DocumentID, FTabPayload_WorkspaceDocument::GetExport(Info.Payload.ToSharedRef())));
 		}
 	}
 
@@ -99,7 +103,7 @@ TAttribute<FText> FAssetDocumentSummoner::ConstructTabNameForObject(UObject* Doc
 
 bool FAssetDocumentSummoner::IsPayloadSupported(TSharedRef<FTabPayload> Payload) const
 {
-	UObject* Object = Payload->IsValid() ? FTabPayload_UObject::CastChecked<UObject>(Payload) : nullptr;
+	UObject* Object = FTabPayload_WorkspaceDocument::CastChecked<UObject>(Payload);
 	if(Object)
 	{
 		FWorkspaceEditorModule& WorkspaceEditorModule = FModuleManager::LoadModuleChecked<FWorkspaceEditorModule>("WorkspaceEditor");
@@ -114,7 +118,7 @@ bool FAssetDocumentSummoner::IsPayloadSupported(TSharedRef<FTabPayload> Payload)
 
 TAttribute<FText> FAssetDocumentSummoner::ConstructTabLabelSuffix(const FWorkflowTabSpawnInfo& Info) const
 {
-	UObject* Object = Info.Payload->IsValid() ? FTabPayload_UObject::CastChecked<UObject>(Info.Payload) : nullptr;
+	UObject* Object = FTabPayload_WorkspaceDocument::CastChecked<UObject>(Info.Payload.ToSharedRef());
 	if(Object)
 	{
 		return MakeAttributeLambda([WeakObject = TWeakObjectPtr<UObject>(Object)]()
@@ -134,38 +138,42 @@ TAttribute<FText> FAssetDocumentSummoner::ConstructTabLabelSuffix(const FWorkflo
 	return FText::GetEmpty();
 }
 
-TSharedRef<SWidget> FAssetDocumentSummoner::CreateTabBodyForObject(const FWorkflowTabSpawnInfo& Info, UObject* DocumentID) const
+TSharedRef<SWidget> FAssetDocumentSummoner::CreateTabBody(const FWorkflowTabSpawnInfo& Info) const
 {
 	TSharedPtr<SWidget> TabContent = SNullWidget::NullWidget;
 
 	const TSharedPtr<FWorkspaceEditor> WorkspaceEditor = HostingAppPtr.Pin();
 	check(WorkspaceEditor.IsValid());
-	
+
+	UObject* DocumentID = FTabPayload_WorkspaceDocument::CastChecked<UObject>(Info.Payload.ToSharedRef());
 	const FWorkspaceEditorModule& WorkspaceEditorModule = FModuleManager::LoadModuleChecked<FWorkspaceEditorModule>("WorkspaceEditor");
+
+	FWorkspaceEditorContext Context(WorkspaceEditor.ToSharedRef(), DocumentID, FTabPayload_WorkspaceDocument::GetExport(Info.Payload.ToSharedRef()));
 	if(const FObjectDocumentArgs* DocumentArgs = WorkspaceEditorModule.FindObjectDocumentType(DocumentID))
 	{
 		if(DocumentArgs->OnMakeDocumentWidget.IsBound())
 		{
-			TabContent = DocumentArgs->OnMakeDocumentWidget.Execute(FWorkspaceEditorContext(WorkspaceEditor.ToSharedRef(), DocumentID));
+			TabContent = DocumentArgs->OnMakeDocumentWidget.Execute(Context);
 		}
 	}
 
-	return SNew(SWorkspaceTabWrapper, Info.TabInfo, WorkspaceEditor, DocumentID)
+	return SNew(SWorkspaceTabWrapper, Info.TabInfo, Context)
 	[
 		TabContent.ToSharedRef()
 	];
 }
 
-const FSlateBrush* FAssetDocumentSummoner::GetTabIconForObject(const FWorkflowTabSpawnInfo& Info, UObject* DocumentID) const
+const FSlateBrush* FAssetDocumentSummoner::GetTabIcon(const FWorkflowTabSpawnInfo& Info) const
 {
 	if(TSharedPtr<FWorkspaceEditor> WorkspaceEditor = HostingAppPtr.Pin())
 	{
+		UObject* DocumentID = FTabPayload_WorkspaceDocument::CastChecked<UObject>(Info.Payload.ToSharedRef());
 		FWorkspaceEditorModule& WorkspaceEditorModule = FModuleManager::LoadModuleChecked<FWorkspaceEditorModule>("WorkspaceEditor");
 		if(const FObjectDocumentArgs* DocumentArgs = WorkspaceEditorModule.FindObjectDocumentType(DocumentID))
 		{
 			if(DocumentArgs->OnGetTabIcon.IsBound())
 			{
-				return DocumentArgs->OnGetTabIcon.Execute(FWorkspaceEditorContext(WorkspaceEditor.ToSharedRef(), DocumentID));
+				return DocumentArgs->OnGetTabIcon.Execute(FWorkspaceEditorContext(WorkspaceEditor.ToSharedRef(), DocumentID, FTabPayload_WorkspaceDocument::GetExport(Info.Payload.ToSharedRef())));
 			}
 		}
 
diff --git a/Engine/Plugins/Experimental/Workspace/Source/WorkspaceEditor/Private/AssetDocumentSummoner.h b/Engine/Plugins/Experimental/Workspace/Source/WorkspaceEditor/Private/AssetDocumentSummoner.h
--- a/Engine/Plugins/Experimental/Workspace/Source/WorkspaceEditor/Private/AssetDocumentSummoner.h
+++ b/Engine/Plugins/Experimental/Workspace/Source/WorkspaceEditor/Private/AssetDocumentSummoner.h
@@ -3,6 +3,7 @@
 #pragma once
 
 #include "CoreMinimal.h"
+#include "WorkspaceTabPayload.h"
 #include "WorkflowOrientedApp/WorkflowUObjectDocuments.h"
 
 class FWorkflowCentricApplication;
@@ -16,7 +17,7 @@ namespace UE::Workspace
 namespace UE::Workspace
 {
 
-struct FAssetDocumentSummoner : public FDocumentTabFactoryForObjects<UObject>
+struct FAssetDocumentSummoner : public FDocumentTabFactory
 {
 public:
 	// Delegate called to save the state of a document
@@ -34,11 +35,20 @@ private:
 	virtual void SaveState(TSharedPtr<SDockTab> Tab, TSharedPtr<FTabPayload> Payload) const override;
 	virtual bool IsPayloadSupported(TSharedRef<FTabPayload> Payload) const override;
 	virtual TAttribute<FText> ConstructTabLabelSuffix(const FWorkflowTabSpawnInfo& Info) const override;
+	virtual TSharedRef<SWidget> CreateTabBody(const FWorkflowTabSpawnInfo& Info) const override;
+	virtual const FSlateBrush* GetTabIcon(const FWorkflowTabSpawnInfo& Info) const override;
 
-	// FDocumentTabFactoryForObjects interface
-	virtual TAttribute<FText> ConstructTabNameForObject(UObject* DocumentID) const override;
-	virtual const FSlateBrush* GetTabIconForObject(const FWorkflowTabSpawnInfo& Info, UObject* DocumentID) const override;
-	virtual TSharedRef<SWidget> CreateTabBodyForObject(const FWorkflowTabSpawnInfo& Info, UObject* DocumentID) const override;
+	virtual bool IsPayloadValid(TSharedRef<FTabPayload> Payload) const override
+	{
+		if (Payload->PayloadType == FTabPayload_WorkspaceDocument::DocumentPayloadName)
+		{
+			return Payload->IsValid();
+		}
+		return false;
+	}
+
+	// Creates the label for the tab
+	virtual TAttribute<FText> ConstructTabName(const FWorkflowTabSpawnInfo& Info) const override;
 
 	// The hosting app
 	TWeakPtr<FWorkspaceEditor> HostingAppPtr;
diff --git a/Engine/Plugins/Experimental/Workspace/Source/WorkspaceEditor/Private/GraphDocumentState.h b/Engine/Plugins/Experimental/Workspace/Source/WorkspaceEditor/Private/GraphDocumentState.h
--- a/Engine/Plugins/Experimental/Workspace/Source/WorkspaceEditor/Private/GraphDocumentState.h
+++ b/Engine/Plugins/Experimental/Workspace/Source/WorkspaceEditor/Private/GraphDocumentState.h
@@ -13,8 +13,8 @@ struct FGraphDocumentState : public FWorkspaceDocumentState
 
 	FGraphDocumentState() = default;
 
-	FGraphDocumentState(const UObject* InObject, const FVector2D& InViewLocation, float InZoomAmount)
-		: FWorkspaceDocumentState(InObject)
+	FGraphDocumentState(const UObject* InObject, const FWorkspaceOutlinerItemExport& InExport, const FVector2D& InViewLocation, float InZoomAmount)
+		: FWorkspaceDocumentState(InObject, InExport)
 		, ViewLocation(InViewLocation)
 		, ZoomAmount(InZoomAmount)
 	{}
diff --git a/Engine/Plugins/Experimental/Workspace/Source/WorkspaceEditor/Private/Outliner/WorkspaceOutlinerMode.cpp b/Engine/Plugins/Experimental/Workspace/Source/WorkspaceEditor/Private/Outliner/WorkspaceOutlinerMode.cpp
--- a/Engine/Plugins/Experimental/Workspace/Source/WorkspaceEditor/Private/Outliner/WorkspaceOutlinerMode.cpp
+++ b/Engine/Plugins/Experimental/Workspace/Source/WorkspaceEditor/Private/Outliner/WorkspaceOutlinerMode.cpp
@@ -79,7 +79,7 @@ TSharedPtr<SWidget> FWorkspaceOutlinerMode::CreateContextMenu()
 
 					const bool bSelectionContainsValidAssetPath = MenuContext->SelectedExports.Num() && MenuContext->SelectedExports.ContainsByPredicate([](const FWorkspaceOutlinerItemExport& Export)
 					{
-						return Export.GetFirstAssetPath().IsValid();
+						return Export.GetFirstAssetPath().IsValid() || Export.GetResolvedExport().GetFirstAssetPath().IsValid();
 					});
 					
 					FToolMenuSection& AssetsSection = InMenu->AddSection("Assets", LOCTEXT("AssetSectionLabel", "Assets"));
@@ -249,11 +249,6 @@ void FWorkspaceOutlinerMode::HandleItemSelection(const FSceneOutlinerItemSelecti
 					[](const FSceneOutlinerTreeItemPtr& InItem)
 					{
 						FWorkspaceOutlinerTreeItem* TreeItem = InItem->CastTo<FWorkspaceOutlinerTreeItem>();
-						if (TreeItem->Export.HasData() && TreeItem->Export.GetData().GetScriptStruct() == FWorkspaceOutlinerAssetReferenceItemData::StaticStruct())
-						{
-							return TreeItem->Export.GetData().GetMutable<FWorkspaceOutlinerAssetReferenceItemData>().ReferredExport;
-						}
-
 						return TreeItem->Export;
 					});
 
@@ -406,11 +401,6 @@ void FWorkspaceOutlinerMode::OpenItems(TArrayView<const FSceneOutlinerTreeItemPt
 			FWorkspaceOutlinerItemExport Export = [&TreeItem]()
 			{
 				const FWorkspaceOutlinerTreeItem* OutlinerTreeItem = TreeItem->CastTo<FWorkspaceOutlinerTreeItem>();
-				if (OutlinerTreeItem->Export.HasData() && OutlinerTreeItem->Export.GetData().GetScriptStruct() == FWorkspaceOutlinerAssetReferenceItemData::StaticStruct())
-				{
-					return OutlinerTreeItem->Export.GetData().Get<FWorkspaceOutlinerAssetReferenceItemData>().ReferredExport;
-				}
-
 				return TreeItem->Export;
 			}();
 			
@@ -429,7 +419,7 @@ void FWorkspaceOutlinerMode::OpenItems(TArrayView<const FSceneOutlinerTreeItemPt
 			{
 				if (const TSharedPtr<UE::Workspace::IWorkspaceEditor> SharedWorkspaceEditor = StaticCastSharedPtr<UE::Workspace::IWorkspaceEditor>(WeakWorkspaceEditor.Pin()))
 				{
-					SharedWorkspaceEditor->OpenAssets({TreeItem->Export.GetFirstAssetPath().TryLoad()});
+					SharedWorkspaceEditor->OpenExports({TreeItem->Export});
 				}
 			}
 		}
diff --git a/Engine/Plugins/Experimental/Workspace/Source/WorkspaceEditor/Private/SGraphDocument.cpp b/Engine/Plugins/Experimental/Workspace/Source/WorkspaceEditor/Private/SGraphDocument.cpp
--- a/Engine/Plugins/Experimental/Workspace/Source/WorkspaceEditor/Private/SGraphDocument.cpp
+++ b/Engine/Plugins/Experimental/Workspace/Source/WorkspaceEditor/Private/SGraphDocument.cpp
@@ -10,10 +10,11 @@
 namespace UE::Workspace
 {
 
-void SGraphDocument::Construct(const FArguments& InArgs, TSharedRef<FWorkspaceEditor> InHostingApp, UEdGraph* InGraph)
+void SGraphDocument::Construct(const FArguments& InArgs, TSharedRef<FWorkspaceEditor> InHostingApp, const FWorkspaceDocument& InWorkspaceDocument)
 {
 	HostingAppPtr = InHostingApp;
-	EdGraph = InGraph;
+	EdGraph = InWorkspaceDocument.GetTypedObject<UEdGraph>();
+	Document = InWorkspaceDocument;
 
 	BindCommands();
 
@@ -33,13 +34,13 @@ void SGraphDocument::Construct(const FArguments& InArgs, TSharedRef<FWorkspaceEd
 	{
 		if(OnCreateActionMenu.IsBound())
 		{
-			return OnCreateActionMenu.Execute(FWorkspaceEditorContext(HostingAppPtr.Pin().ToSharedRef(), EdGraph), InGraph, InNodePosition, InDraggedPins, bInAutoExpand, InOnMenuClosed);
+			return OnCreateActionMenu.Execute(FWorkspaceEditorContext(HostingAppPtr.Pin().ToSharedRef(), Document), InGraph, InNodePosition, InDraggedPins, bInAutoExpand, InOnMenuClosed);
 		}
 		return FActionMenuContent();
 	});
 	Events.OnSelectionChanged = SGraphEditor::FOnSelectionChanged::CreateLambda([this, OnGraphSelectionChanged = InArgs._OnGraphSelectionChanged](const TSet<UObject*>& NewSelection)
 	{
-		OnGraphSelectionChanged.ExecuteIfBound(FWorkspaceEditorContext(HostingAppPtr.Pin().ToSharedRef(), EdGraph), NewSelection);
+		OnGraphSelectionChanged.ExecuteIfBound(FWorkspaceEditorContext(HostingAppPtr.Pin().ToSharedRef(), Document), NewSelection);
 
 		if (const TSharedPtr<FWorkspaceEditor> SharedWorkspaceEditor = HostingAppPtr.Pin())
 		{
@@ -48,19 +49,19 @@ void SGraphDocument::Construct(const FArguments& InArgs, TSharedRef<FWorkspaceEd
 	});
 	Events.OnTextCommitted = ::FOnNodeTextCommitted::CreateLambda([this, OnNodeTextCommitted = InArgs._OnNodeTextCommitted](const FText& NewText, ETextCommit::Type CommitInfo, UEdGraphNode* NodeBeingChanged)
 	{
-		OnNodeTextCommitted.ExecuteIfBound(FWorkspaceEditorContext(HostingAppPtr.Pin().ToSharedRef(), EdGraph), NewText, CommitInfo, NodeBeingChanged);
+		OnNodeTextCommitted.ExecuteIfBound(FWorkspaceEditorContext(HostingAppPtr.Pin().ToSharedRef(), Document), NewText, CommitInfo, NodeBeingChanged);
 	});
 	Events.OnNodeDoubleClicked = FSingleNodeEvent::CreateLambda([this, OnNodeDoubleClicked = InArgs._OnNodeDoubleClicked](UEdGraphNode* InNode)
 	{
-		OnNodeDoubleClicked.ExecuteIfBound(FWorkspaceEditorContext(HostingAppPtr.Pin().ToSharedRef(), EdGraph), InNode);
+		OnNodeDoubleClicked.ExecuteIfBound(FWorkspaceEditorContext(HostingAppPtr.Pin().ToSharedRef(), Document), InNode);
 	});
 
 	ChildSlot
 	[
 		SAssignNew(GraphEditor, SGraphEditor)
 		.AdditionalCommands(CommandList)
-		.IsEditable(this, &SGraphDocument::IsEditable, InGraph)
-		.GraphToEdit(InGraph)
+		.IsEditable(this, &SGraphDocument::IsEditable, EdGraph)
+		.GraphToEdit(EdGraph)
 		.GraphEvents(Events)
 		.AssetEditorToolkit(HostingAppPtr)
 		.OnNavigateHistoryBack(InArgs._OnNavigateHistoryBack)		
@@ -99,8 +100,8 @@ void SGraphDocument::BindCommands()
 
 bool SGraphDocument::CanDeleteSelectedNodes() const
 {
-	const bool bCan = OnCanCutSelectedNodes.IsBound()
-		? IsEditable(EdGraph) && OnCanCutSelectedNodes.Execute(FWorkspaceEditorContext(HostingAppPtr.Pin().ToSharedRef(), EdGraph), GraphEditor->GetSelectedNodes())
+	const bool bCan = OnCanDeleteSelectedNodes.IsBound()
+		? IsEditable(EdGraph) && OnCanDeleteSelectedNodes.Execute(FWorkspaceEditorContext(HostingAppPtr.Pin().ToSharedRef(), Document), GraphEditor->GetSelectedNodes())
 		: false;
 	return bCan;
 }
@@ -108,13 +109,13 @@ bool SGraphDocument::CanDeleteSelectedNodes() const
 void SGraphDocument::DeleteSelectedNodes()
 {
 	const FGraphPanelSelectionSet SelectedNodes = GraphEditor->GetSelectedNodes();
-	OnDeleteSelectedNodes.ExecuteIfBound(FWorkspaceEditorContext(HostingAppPtr.Pin().ToSharedRef(), EdGraph), SelectedNodes);
+	OnDeleteSelectedNodes.ExecuteIfBound(FWorkspaceEditorContext(HostingAppPtr.Pin().ToSharedRef(), Document), SelectedNodes);
 }
 
 bool SGraphDocument::CanCutSelectedNodes() const
 {
 	const bool bCan = OnCanCutSelectedNodes.IsBound() 
-		? IsEditable(EdGraph) && OnCanCutSelectedNodes.Execute(FWorkspaceEditorContext(HostingAppPtr.Pin().ToSharedRef(), EdGraph), GraphEditor->GetSelectedNodes())
+		? IsEditable(EdGraph) && OnCanCutSelectedNodes.Execute(FWorkspaceEditorContext(HostingAppPtr.Pin().ToSharedRef(), Document), GraphEditor->GetSelectedNodes())
 		: false;
 	return bCan;
 }
@@ -122,13 +123,13 @@ bool SGraphDocument::CanCutSelectedNodes() const
 void SGraphDocument::CutSelectedNodes()
 {
 	const FGraphPanelSelectionSet SelectedNodes = GraphEditor->GetSelectedNodes();
-	OnCutSelectedNodes.ExecuteIfBound(FWorkspaceEditorContext(HostingAppPtr.Pin().ToSharedRef(), EdGraph), SelectedNodes);
+	OnCutSelectedNodes.ExecuteIfBound(FWorkspaceEditorContext(HostingAppPtr.Pin().ToSharedRef(), Document), SelectedNodes);
 }
 
 bool SGraphDocument::CanCopySelectedNodes() const
 {
 	const bool bCan = OnCanCopySelectedNodes.IsBound()
-		? OnCanCopySelectedNodes.Execute(FWorkspaceEditorContext(HostingAppPtr.Pin().ToSharedRef(), EdGraph), GraphEditor->GetSelectedNodes())
+		? OnCanCopySelectedNodes.Execute(FWorkspaceEditorContext(HostingAppPtr.Pin().ToSharedRef(), Document), GraphEditor->GetSelectedNodes())
 		: false;
 	return bCan;
 }
@@ -136,7 +137,7 @@ bool SGraphDocument::CanCopySelectedNodes() const
 void SGraphDocument::CopySelectedNodes()
 {
 	const FGraphPanelSelectionSet SelectedNodes = GraphEditor->GetSelectedNodes();
-	OnCopySelectedNodes.ExecuteIfBound(FWorkspaceEditorContext(HostingAppPtr.Pin().ToSharedRef(), EdGraph), SelectedNodes);
+	OnCopySelectedNodes.ExecuteIfBound(FWorkspaceEditorContext(HostingAppPtr.Pin().ToSharedRef(), Document), SelectedNodes);
 }
 
 bool SGraphDocument::CanPasteNodes() const
@@ -145,7 +146,7 @@ bool SGraphDocument::CanPasteNodes() const
 	FPlatformApplicationMisc::ClipboardPaste(TextToImport);
 
 	const bool bCan = OnCanPasteNodes.IsBound()
-		? IsEditable(EdGraph) && !TextToImport.IsEmpty() && OnCanPasteNodes.Execute(FWorkspaceEditorContext(HostingAppPtr.Pin().ToSharedRef(), EdGraph), TextToImport)
+		? IsEditable(EdGraph) && !TextToImport.IsEmpty() && OnCanPasteNodes.Execute(FWorkspaceEditorContext(HostingAppPtr.Pin().ToSharedRef(), Document), TextToImport)
 		: false;
 	return bCan;
 }
@@ -155,13 +156,13 @@ void SGraphDocument::PasteNodes()
 	FString TextToImport;
 	FPlatformApplicationMisc::ClipboardPaste(TextToImport);
 
-	OnPasteNodes.ExecuteIfBound(FWorkspaceEditorContext(HostingAppPtr.Pin().ToSharedRef(), EdGraph), GraphEditor->GetPasteLocation(), TextToImport);
+	OnPasteNodes.ExecuteIfBound(FWorkspaceEditorContext(HostingAppPtr.Pin().ToSharedRef(), Document), GraphEditor->GetPasteLocation(), TextToImport);
 }
 
 bool SGraphDocument::CanDuplicateSelectedNodes() const
 {
 	const bool bCan = OnCanDuplicateSelectedNodes.IsBound()
-		? IsEditable(EdGraph) && OnCanDuplicateSelectedNodes.Execute(FWorkspaceEditorContext(HostingAppPtr.Pin().ToSharedRef(), EdGraph), GraphEditor->GetSelectedNodes())
+		? IsEditable(EdGraph) && OnCanDuplicateSelectedNodes.Execute(FWorkspaceEditorContext(HostingAppPtr.Pin().ToSharedRef(), Document), GraphEditor->GetSelectedNodes())
 		: false;
 	return bCan;
 }
@@ -169,7 +170,7 @@ bool SGraphDocument::CanDuplicateSelectedNodes() const
 void SGraphDocument::DuplicateSelectedNodes()
 {
 	const FGraphPanelSelectionSet SelectedNodes = GraphEditor->GetSelectedNodes();
-	OnDuplicateSelectedNodes.ExecuteIfBound(FWorkspaceEditorContext(HostingAppPtr.Pin().ToSharedRef(), EdGraph), GraphEditor->GetPasteLocation(), SelectedNodes);
+	OnDuplicateSelectedNodes.ExecuteIfBound(FWorkspaceEditorContext(HostingAppPtr.Pin().ToSharedRef(), Document), GraphEditor->GetPasteLocation(), SelectedNodes);
 }
 
 bool SGraphDocument::CanSelectAllNodes() const
@@ -183,7 +184,6 @@ void SGraphDocument::SelectAllNodes()
 	GraphEditor->SelectAllNodes();
 }
 
-
 bool SGraphDocument::IsEditable(UEdGraph* InGraph) const
 {
 	return InGraph && HostingAppPtr.Pin()->InEditingMode() && InGraph->bEditable;
diff --git a/Engine/Plugins/Experimental/Workspace/Source/WorkspaceEditor/Private/SGraphDocument.h b/Engine/Plugins/Experimental/Workspace/Source/WorkspaceEditor/Private/SGraphDocument.h
--- a/Engine/Plugins/Experimental/Workspace/Source/WorkspaceEditor/Private/SGraphDocument.h
+++ b/Engine/Plugins/Experimental/Workspace/Source/WorkspaceEditor/Private/SGraphDocument.h
@@ -48,7 +48,7 @@ class SGraphDocument : public SCompoundWidget
 
 	SLATE_END_ARGS()
 
-	void Construct(const FArguments& InArgs, TSharedRef<FWorkspaceEditor> InHostingApp, UEdGraph* InGraph);
+	void Construct(const FArguments& InArgs, TSharedRef<FWorkspaceEditor> InHostingApp, const FWorkspaceDocument& InDocument);
 	
 	void BindCommands();
 	
@@ -77,6 +77,7 @@ class SGraphDocument : public SCompoundWidget
 
 	// The graph we are editing
 	UEdGraph* EdGraph = nullptr;
+	FWorkspaceDocument Document;
 
 	// The graph editor we wrap
 	TSharedPtr<SGraphEditor> GraphEditor;
diff --git a/Engine/Plugins/Experimental/Workspace/Source/WorkspaceEditor/Private/SWorkspaceTabWrapper.cpp b/Engine/Plugins/Experimental/Workspace/Source/WorkspaceEditor/Private/SWorkspaceTabWrapper.cpp
--- a/Engine/Plugins/Experimental/Workspace/Source/WorkspaceEditor/Private/SWorkspaceTabWrapper.cpp
+++ b/Engine/Plugins/Experimental/Workspace/Source/WorkspaceEditor/Private/SWorkspaceTabWrapper.cpp
@@ -10,11 +10,12 @@
 
 #define LOCTEXT_NAMESPACE "SWorkspaceTabWrapper"
 
-void SWorkspaceTabWrapper::Construct( const FArguments& InArgs, TSharedPtr<class FTabInfo> InTabInfo, TSharedPtr<UE::Workspace::FWorkspaceEditor> InWorkspaceEditor, UObject* InDocumentID)
+void SWorkspaceTabWrapper::Construct( const FArguments& InArgs, TSharedPtr<class FTabInfo> InTabInfo, const UE::Workspace::FWorkspaceEditorContext& InEditorContext)
 {
 	Content = InArgs._Content.Widget;
-	WeakWorkspaceEditor = InWorkspaceEditor;
-	WeakDocumentObject = InDocumentID;
+	WeakWorkspaceEditor = StaticCastSharedRef<UE::Workspace::FWorkspaceEditor>(InEditorContext.WorkspaceEditor);
+	WeakDocumentObject = InEditorContext.Document.Object;
+	Export = InEditorContext.Document.Export;
 
 	if (InTabInfo->GetTab().IsValid())
 	{
@@ -224,20 +225,56 @@ void SWorkspaceTabWrapper::RebuildBreadcrumbTrail() const
 	BreadcrumbTrail->ClearCrumbs(false);
 
 	TArray<TSharedPtr<UE::Workspace::FWorkspaceBreadcrumb>> Breadcrumbs;
-	if (UObject* DocumentID = WeakDocumentObject.Get())
+	TArray<FWorkspaceOutlinerItemExport> Exports;
+
+	if (Export.GetFirstAssetPath().IsValid())
 	{
-		if (const TSharedPtr<UE::Workspace::FWorkspaceEditor> SharedWorkspaceEditor = WeakWorkspaceEditor.Pin())
+		if (FWorkspaceOutlinerAssetReferenceItemData::IsAssetReference(Export))
 		{
-			const UE::Workspace::FWorkspaceEditorModule& WorkspaceEditorModule = FModuleManager::LoadModuleChecked<UE::Workspace::FWorkspaceEditorModule>("WorkspaceEditor");
-			if(const UE::Workspace::FObjectDocumentArgs* DocumentArgs = WorkspaceEditorModule.FindObjectDocumentType(DocumentID))
+			// Prepend outer asset/export chain
+			TArray<FSoftObjectPath> AssetPaths;
+			Export.GetAssetPaths(AssetPaths);
+			Export.GetExports(Exports);
+		}
+		else
+		{
+			Exports.Add(Export);
+		}
+
+		for (const FWorkspaceOutlinerItemExport& DocumentExport : Exports)
+		{
+			UObject* DocumentID = DocumentExport.GetFirstAssetOfType<UObject>();
+			check(DocumentID);
+			if (const TSharedPtr<UE::Workspace::FWorkspaceEditor> SharedWorkspaceEditor = WeakWorkspaceEditor.Pin())
 			{
-				if(DocumentArgs->OnGetDocumentBreadcrumbTrail.IsBound())
+				const UE::Workspace::FWorkspaceEditorModule& WorkspaceEditorModule = FModuleManager::LoadModuleChecked<UE::Workspace::FWorkspaceEditorModule>("WorkspaceEditor");
+				if(const UE::Workspace::FObjectDocumentArgs* DocumentArgs = WorkspaceEditorModule.FindObjectDocumentType(DocumentID))
 				{
-					DocumentArgs->OnGetDocumentBreadcrumbTrail.Execute(UE::Workspace::FWorkspaceEditorContext(SharedWorkspaceEditor.ToSharedRef(), DocumentID), Breadcrumbs);
+					if(DocumentArgs->OnGetDocumentBreadcrumbTrail.IsBound())
+					{
+						DocumentArgs->OnGetDocumentBreadcrumbTrail.Execute(UE::Workspace::FWorkspaceEditorContext(SharedWorkspaceEditor.ToSharedRef(), { DocumentExport, DocumentID}), Breadcrumbs);
+					}
 				}
 			}
 		}
 	}
+	else
+	{
+		if (UObject* DocumentID = WeakDocumentObject.Get())
+	    {
+		    if (const TSharedPtr<UE::Workspace::FWorkspaceEditor> SharedWorkspaceEditor = WeakWorkspaceEditor.Pin())
+		    {
+			    const UE::Workspace::FWorkspaceEditorModule& WorkspaceEditorModule = FModuleManager::LoadModuleChecked<UE::Workspace::FWorkspaceEditorModule>("WorkspaceEditor");
+			    if(const UE::Workspace::FObjectDocumentArgs* DocumentArgs = WorkspaceEditorModule.FindObjectDocumentType(DocumentID))
+			    {
+				    if(DocumentArgs->OnGetDocumentBreadcrumbTrail.IsBound())
+				    {
+					    DocumentArgs->OnGetDocumentBreadcrumbTrail.Execute(UE::Workspace::FWorkspaceEditorContext(SharedWorkspaceEditor.ToSharedRef(), { FWorkspaceOutlinerItemExport(), DocumentID}), Breadcrumbs);
+				    }
+			    }
+		    }
+	    }
+	}	
 
 	// Widgets have to be added in reverse order
 	for (int32 Index = Breadcrumbs.Num() - 1; Index >= 0; --Index)
@@ -258,7 +295,7 @@ const FSlateBrush* SWorkspaceTabWrapper::GetTabIcon() const
 			{
 				if(DocumentArgs->OnGetTabIcon.IsBound())
 				{
-					return DocumentArgs->OnGetTabIcon.Execute(UE::Workspace::FWorkspaceEditorContext(SharedWorkspaceEditor.ToSharedRef(), DocumentID));
+					return DocumentArgs->OnGetTabIcon.Execute(UE::Workspace::FWorkspaceEditorContext(SharedWorkspaceEditor.ToSharedRef(), { Export, DocumentID }));
 				}
 			}
 		}			
diff --git a/Engine/Plugins/Experimental/Workspace/Source/WorkspaceEditor/Private/SWorkspaceTabWrapper.h b/Engine/Plugins/Experimental/Workspace/Source/WorkspaceEditor/Private/SWorkspaceTabWrapper.h
--- a/Engine/Plugins/Experimental/Workspace/Source/WorkspaceEditor/Private/SWorkspaceTabWrapper.h
+++ b/Engine/Plugins/Experimental/Workspace/Source/WorkspaceEditor/Private/SWorkspaceTabWrapper.h
@@ -18,7 +18,7 @@ public:
 		SLATE_DEFAULT_SLOT(FArguments, Content)
 	SLATE_END_ARGS()
 	
-	void Construct(const FArguments& InArgs, TSharedPtr<class FTabInfo> InTabInfo, TSharedPtr<UE::Workspace::FWorkspaceEditor> InWorkspaceEditor, UObject* InDocumentID);
+	void Construct(const FArguments& InArgs, TSharedPtr<class FTabInfo> InTabInfo, const UE::Workspace::FWorkspaceEditorContext& InEditorContext);
 	TSharedRef<SWidget> GetContent() const { return Content.ToSharedRef(); }
 	TWeakObjectPtr<UObject> GetDocumentObject() const { return WeakDocumentObject; }
 	
@@ -35,6 +35,7 @@ protected:
 	TSharedPtr<SWidget> Content = nullptr;
 	TWeakPtr<UE::Workspace::FWorkspaceEditor> WeakWorkspaceEditor = nullptr;
 	TWeakObjectPtr<UObject> WeakDocumentObject = nullptr;
+	FWorkspaceOutlinerItemExport Export;
 
 	TSharedPtr<SScrollBox> BreadcrumbTrailScrollBox = nullptr;
 	TSharedPtr<SBreadcrumbTrail<TSharedPtr<UE::Workspace::FWorkspaceBreadcrumb>>> BreadcrumbTrail = nullptr;	
diff --git a/Engine/Plugins/Experimental/Workspace/Source/WorkspaceEditor/Private/SWorkspaceView.cpp b/Engine/Plugins/Experimental/Workspace/Source/WorkspaceEditor/Private/SWorkspaceView.cpp
--- a/Engine/Plugins/Experimental/Workspace/Source/WorkspaceEditor/Private/SWorkspaceView.cpp
+++ b/Engine/Plugins/Experimental/Workspace/Source/WorkspaceEditor/Private/SWorkspaceView.cpp
@@ -15,6 +15,7 @@
 #include "Framework/Commands/GenericCommands.h"
 #include "Outliner/WorkspaceOutlinerColumns.h"
 #include "Outliner/WorkspaceOutlinerMode.h"
+#include "Outliner/WorkspaceOutlinerTreeItem.h"
 
 #define LOCTEXT_NAMESPACE "SWorkspaceView"
 
@@ -168,7 +169,7 @@ void SWorkspaceView::Construct(const FArguments& InArgs, UWorkspace* InWorkspace
 	];
 }
 
-void SWorkspaceView::SelectObject(UObject* InObject)
+void SWorkspaceView::SelectObject(UObject* InObject) const
 {
 	FWorkspaceOutlinerItemExport Export(InObject->GetFName(), InObject);
 
@@ -179,6 +180,26 @@ void SWorkspaceView::SelectObject(UObject* InObject)
 	}
 }
 
+void SWorkspaceView::SelectExport(const FWorkspaceOutlinerItemExport& InExport) const
+{
+	FSceneOutlinerTreeItemPtr FoundItem = SceneWorkspaceOutliner->GetTreeItem(FSceneOutlinerTreeItemID(GetTypeHash(InExport)));
+	if(FoundItem.IsValid())
+	{
+		SceneWorkspaceOutliner->SetItemSelection(FoundItem, true, ESelectInfo::OnMouseClick);	// Not direct, so we get callbacks
+	}
+}
+
+void SWorkspaceView::GetWorkspaceExportData(FWorkspaceOutlinerItemExport& InOutPartialExport) const
+{
+	FSceneOutlinerTreeItemPtr FoundItem = SceneWorkspaceOutliner->GetTreeItem(FSceneOutlinerTreeItemID(GetTypeHash(InOutPartialExport)));
+	if(FoundItem.IsValid())
+	{
+		if (const FWorkspaceOutlinerTreeItem* TreeItem = FoundItem->CastTo<FWorkspaceOutlinerTreeItem>())
+		{
+			InOutPartialExport = TreeItem->Export;
+		}
+	}
+}
 }
 
 #undef LOCTEXT_NAMESPACE
\ No newline at end of file
diff --git a/Engine/Plugins/Experimental/Workspace/Source/WorkspaceEditor/Private/SWorkspaceView.h b/Engine/Plugins/Experimental/Workspace/Source/WorkspaceEditor/Private/SWorkspaceView.h
--- a/Engine/Plugins/Experimental/Workspace/Source/WorkspaceEditor/Private/SWorkspaceView.h
+++ b/Engine/Plugins/Experimental/Workspace/Source/WorkspaceEditor/Private/SWorkspaceView.h
@@ -5,6 +5,7 @@
 #include "ContentBrowserDelegates.h"
 #include "Widgets/SCompoundWidget.h"
 
+struct FWorkspaceOutlinerItemExport;
 class UWorkspace;
 class UWorkSpaceAssetUserData;
 
@@ -20,7 +21,10 @@ class SWorkspaceView : public SCompoundWidget
 
 	void Construct(const FArguments& InArgs, UWorkspace* InWorkspace, TSharedRef<UE::Workspace::IWorkspaceEditor> InWorkspaceEditor);
 
-	void SelectObject(UObject* InObject);
+	void SelectObject(UObject* InObject) const;
+	void SelectExport(const FWorkspaceOutlinerItemExport& InExport) const;
+	// Tries to lookup the FWorkspaceOutlinerItemExport path within the outliner, and copies over its TInstancedStruct<FWorkspaceOutlinerItemData> value
+	void GetWorkspaceExportData(FWorkspaceOutlinerItemExport& InOutPartialExport) const;
 	
 private:
 	UWorkspace* Workspace = nullptr;
diff --git a/Engine/Plugins/Experimental/Workspace/Source/WorkspaceEditor/Private/WorkspaceAssetRegistryInfo.cpp b/Engine/Plugins/Experimental/Workspace/Source/WorkspaceEditor/Private/WorkspaceAssetRegistryInfo.cpp
new file mode 100644
--- /dev/null
+++ b/Engine/Plugins/Experimental/Workspace/Source/WorkspaceEditor/Private/WorkspaceAssetRegistryInfo.cpp
@@ -0,0 +1,23 @@
+// Copyright Epic Games, Inc. All Rights Reserved.
+
+#include "WorkspaceAssetRegistryInfo.h"
+
+FWorkspaceOutlinerItemExport& FWorkspaceOutlinerItemExport::GetResolvedExport()
+{
+	if (HasData() && Data.GetScriptStruct() == FWorkspaceOutlinerAssetReferenceItemData::StaticStruct())
+	{
+		return Data.GetMutable<FWorkspaceOutlinerAssetReferenceItemData>().ReferredExport;
+	}
+
+	return *this;
+}
+
+const FWorkspaceOutlinerItemExport& FWorkspaceOutlinerItemExport::GetResolvedExport() const
+{
+	if (HasData() && Data.GetScriptStruct() == FWorkspaceOutlinerAssetReferenceItemData::StaticStruct())
+	{
+		return Data.Get<FWorkspaceOutlinerAssetReferenceItemData>().ReferredExport;
+	}
+
+	return *this;
+}
diff --git a/Engine/Plugins/Experimental/Workspace/Source/WorkspaceEditor/Private/WorkspaceEditor.cpp b/Engine/Plugins/Experimental/Workspace/Source/WorkspaceEditor/Private/WorkspaceEditor.cpp
--- a/Engine/Plugins/Experimental/Workspace/Source/WorkspaceEditor/Private/WorkspaceEditor.cpp
+++ b/Engine/Plugins/Experimental/Workspace/Source/WorkspaceEditor/Private/WorkspaceEditor.cpp
@@ -26,6 +26,8 @@
 #include "ContextObjectStore.h"
 #include "SGraphDocument.h"
 #include "WorkspaceSchema.h"
+#include "WorkspaceTabPayload.h"
+#include "SWorkspaceView.h"
 
 #define LOCTEXT_NAMESPACE "WorkspaceEditor"
 
@@ -251,14 +253,15 @@ void FWorkspaceEditor::RestoreEditedObjectState()
 	{
 		if (UObject* Object = DocumentState.Get().Object.TryLoad())
 		{
-			if(const TSharedPtr<SDockTab> DockTab = OpenDocument(Object, FDocumentTracker::RestorePreviousDocument))
+			const FWorkspaceOutlinerItemExport& Export = DocumentState.Get().Export;
+			if(const TSharedPtr<SDockTab> DockTab = OpenDocument(Object, Export, FDocumentTracker::RestorePreviousDocument))
 			{
 				const TSharedRef<SWorkspaceTabWrapper> TabWrapper = StaticCastSharedRef<SWorkspaceTabWrapper>(DockTab->GetContent());	
 				FWorkspaceEditorModule& WorkspaceEditorModule = FModuleManager::LoadModuleChecked<FWorkspaceEditorModule>("WorkspaceEditor");
 				const FObjectDocumentArgs* DocumentArgs = WorkspaceEditorModule.FindObjectDocumentType(Object);
 				if(DocumentArgs != nullptr && DocumentArgs->OnSetDocumentState.IsBound())
 				{
-					DocumentArgs->OnSetDocumentState.Execute(FWorkspaceEditorContext(SharedThis(this), Object), TabWrapper->GetContent(), DocumentState);
+					DocumentArgs->OnSetDocumentState.Execute(FWorkspaceEditorContext(SharedThis(this), { Export, Object }), TabWrapper->GetContent(), DocumentState);
 				}
 			}
 		}
@@ -281,7 +284,13 @@ void FWorkspaceEditor::SaveEditedObjectState() const
 
 TSharedPtr<SDockTab> FWorkspaceEditor::OpenDocument(const UObject* InForObject, FDocumentTracker::EOpenDocumentCause InCause)
 {
-	const TSharedRef<FTabPayload_UObject> Payload = FTabPayload_UObject::Make(InForObject);
+	FWorkspaceOutlinerItemExport Export(InForObject->GetFName(), InForObject->GetPackage());
+	return OpenDocument(InForObject, Export, InCause);
+}
+
+TSharedPtr<SDockTab> FWorkspaceEditor::OpenDocument(const UObject* InForObject, const FWorkspaceOutlinerItemExport& InExport, FDocumentTracker::EOpenDocumentCause InCause)
+{
+	const TSharedRef<FTabPayload_WorkspaceDocument> Payload = FTabPayload_WorkspaceDocument::Make(InForObject, InExport);
 	const bool bIsSupportedDocument = DocumentManager->FindSupportingFactory(Payload).IsValid();
 	if (bIsSupportedDocument)
 	{
@@ -297,11 +306,11 @@ TSharedPtr<SDockTab> FWorkspaceEditor::OpenDocument(const UObject* InForObject,
 		if(DocumentArgs != nullptr && DocumentArgs->OnGetDocumentState.IsBound())
 		{
 			const TSharedRef<SWorkspaceTabWrapper> TabWrapper = StaticCastSharedRef<SWorkspaceTabWrapper>(NewTab->GetContent());
-			RecordDocumentState(DocumentArgs->OnGetDocumentState.Execute(FWorkspaceEditorContext(SharedThis(this), const_cast<UObject*>(InForObject)), TabWrapper->GetContent()));
+			RecordDocumentState(DocumentArgs->OnGetDocumentState.Execute(FWorkspaceEditorContext(SharedThis(this), { InExport, const_cast<UObject*>(InForObject) }), TabWrapper->GetContent()));
 		}
 		else
 		{
-			RecordDocumentState(TInstancedStruct<FWorkspaceDocumentState>::Make(InForObject));
+			RecordDocumentState(TInstancedStruct<FWorkspaceDocumentState>::Make(InForObject, InExport));
 		}
 	}
 
@@ -327,6 +336,30 @@ void FWorkspaceEditor::OpenAssets(TConstArrayView<FAssetData> InAssets)
 	}
 }
 
+void FWorkspaceEditor::OpenExports(TConstArrayView<FWorkspaceOutlinerItemExport> InExports)
+{
+	TArray<const FWorkspaceOutlinerItemExport*> OpenedExports;
+	for(const FWorkspaceOutlinerItemExport& Export : InExports)
+	{
+		FWorkspaceOutlinerItemExport QualifiedExport = Export;
+		WorkspaceView->GetWorkspaceExportData(QualifiedExport);
+		
+		if(UObject* LoadedAsset = QualifiedExport.GetFirstAssetPath().TryLoad())
+		{
+			OpenDocument(LoadedAsset, QualifiedExport, FDocumentTracker::EOpenDocumentCause::NavigatingCurrentDocument);
+			OpenedExports.AddUnique(&QualifiedExport);
+		}
+	}
+
+	if(OpenedExports.Num() > 0)
+	{
+		if(const FWorkspaceOutlinerItemExport* LoadedExport = OpenedExports.Last())
+		{
+			WorkspaceView->SelectExport(*LoadedExport);
+		}
+	}
+}
+
 void FWorkspaceEditor::OpenObjects(TConstArrayView<UObject*> InObjects)
 {
 	for(UObject* Object : InObjects)
@@ -354,7 +387,7 @@ void FWorkspaceEditor::OpenObjects(TConstArrayView<UObject*> InObjects)
 					check(TabContentWidget->GetType() == TEXT("SGraphEditor"));
 				}
 
-				DocumentSubObjectArgs->OnPostDocumentOpenedForSubObject.ExecuteIfBound(FWorkspaceEditorContext(SharedThis(this), Object), TabContentWidget, OriginalObject);
+				DocumentSubObjectArgs->OnPostDocumentOpenedForSubObject.ExecuteIfBound(FWorkspaceEditorContext(SharedThis(this), { FWorkspaceOutlinerItemExport(), Object }), TabContentWidget, OriginalObject);
 			}
 		}
 	}
@@ -652,9 +685,10 @@ void FWorkspaceEditor::SaveAsset_Execute()
 void FWorkspaceEditor::CloseDocumentTab(const UObject* DocumentID) const
 {
 	UWorkspaceState* State = Workspace->GetState();
-	State->DocumentStates.Remove(TInstancedStruct<FWorkspaceDocumentState>::Make(DocumentID));
+	FWorkspaceOutlinerItemExport Export;
+	State->DocumentStates.Remove(TInstancedStruct<FWorkspaceDocumentState>::Make(DocumentID, Export));
 
-	const TSharedRef<FTabPayload_UObject> Payload = FTabPayload_UObject::Make(DocumentID);
+	const TSharedRef<FTabPayload_WorkspaceDocument> Payload = FTabPayload_WorkspaceDocument::Make(DocumentID);
 	DocumentManager->CloseTab(Payload);
 }
 
@@ -734,13 +768,13 @@ void FWorkspaceEditor::RecordDocumentState(const TInstancedStruct<FWorkspaceDocu
 
 void FWorkspaceEditor::NavigateBack()
 {
-	const TSharedRef<FTabPayload_UObject> Payload = FTabPayload_UObject::Make(nullptr);
+	const TSharedRef<FTabPayload_WorkspaceDocument> Payload = FTabPayload_WorkspaceDocument::Make(nullptr);
 	TSharedPtr<SDockTab> OpenedTab = DocumentManager->OpenDocument(Payload, FDocumentTracker::NavigateBackwards);
 }
 
 void FWorkspaceEditor::NavigateForward()
 {
-	const TSharedRef<FTabPayload_UObject> Payload = FTabPayload_UObject::Make(nullptr);
+	const TSharedRef<FTabPayload_WorkspaceDocument> Payload = FTabPayload_WorkspaceDocument::Make(nullptr);
 	TSharedPtr<SDockTab> OpenedTab = DocumentManager->OpenDocument(Payload, FDocumentTracker::NavigateForwards);
 }
 
diff --git a/Engine/Plugins/Experimental/Workspace/Source/WorkspaceEditor/Private/WorkspaceEditor.h b/Engine/Plugins/Experimental/Workspace/Source/WorkspaceEditor/Private/WorkspaceEditor.h
--- a/Engine/Plugins/Experimental/Workspace/Source/WorkspaceEditor/Private/WorkspaceEditor.h
+++ b/Engine/Plugins/Experimental/Workspace/Source/WorkspaceEditor/Private/WorkspaceEditor.h
@@ -95,6 +95,7 @@ private:
 
 	// IWorkspaceEditor interface
 	virtual void OpenAssets(TConstArrayView<FAssetData> InAssets) override;
+	virtual void OpenExports(TConstArrayView<FWorkspaceOutlinerItemExport> InExports) override;
 	virtual void OpenObjects(TConstArrayView<UObject*> InObjects) override;
 	virtual void GetOpenedAssetsOfClass(TSubclassOf<UObject> InClass, TArray<UObject*>& OutObjects) const override;
 	virtual void GetAssets(TArray<FAssetData>& OutAssets) const override;
@@ -128,6 +129,7 @@ private:
 	void SaveEditedObjectState() const;
 
 	TSharedPtr<SDockTab> OpenDocument(const UObject* InForObject, FDocumentTracker::EOpenDocumentCause InCause);
+	TSharedPtr<SDockTab> OpenDocument(const UObject* InForObject, const FWorkspaceOutlinerItemExport& InExport, FDocumentTracker::EOpenDocumentCause InCause);
 
 	void RecordDocumentState(const TInstancedStruct<FWorkspaceDocumentState>& InState) const;
 
diff --git a/Engine/Plugins/Experimental/Workspace/Source/WorkspaceEditor/Private/WorkspaceEditorModule.cpp b/Engine/Plugins/Experimental/Workspace/Source/WorkspaceEditor/Private/WorkspaceEditorModule.cpp
--- a/Engine/Plugins/Experimental/Workspace/Source/WorkspaceEditor/Private/WorkspaceEditorModule.cpp
+++ b/Engine/Plugins/Experimental/Workspace/Source/WorkspaceEditor/Private/WorkspaceEditorModule.cpp
@@ -27,10 +27,9 @@ namespace UE::Workspace
 
 TMap<FOutlinerItemDetailsId, TSharedPtr<IWorkspaceOutlinerItemDetails>> FWorkspaceEditorModule::OutlinerItemDetails;
 
-FWorkspaceEditorContext::FWorkspaceEditorContext(const TSharedRef<IWorkspaceEditor>& InWorkspaceEditor, UObject* InObject)
-	: WorkspaceEditor(InWorkspaceEditor)
-	, Object(InObject)
+FWorkspaceDocument::FWorkspaceDocument(const FWorkspaceOutlinerItemExport& InExport, UObject* InObject) : Export(InExport), Object(InObject)
 {
+	check(InObject);
 	FWorkspaceEditorModule& WorkspaceEditorModule = FModuleManager::Get().LoadModuleChecked<FWorkspaceEditorModule>("WorkspaceEditor");
 	const FObjectDocumentArgs* FoundArgs = WorkspaceEditorModule.ObjectDocumentArgs.Find(InObject->GetClass()->GetClassPathName());
 	while(FoundArgs && FoundArgs->OnRedirectWorkspaceContext.IsBound())
@@ -42,6 +41,14 @@ FWorkspaceEditorContext::FWorkspaceEditorContext(const TSharedRef<IWorkspaceEdit
 	Object = InObject;
 }
 
+FWorkspaceEditorContext::FWorkspaceEditorContext(const TSharedRef<IWorkspaceEditor>& InWorkspaceEditor, UObject* InObject, const FWorkspaceOutlinerItemExport& InExport) : FWorkspaceEditorContext(InWorkspaceEditor, { InExport, InObject })
+{
+}
+
+FWorkspaceEditorContext::FWorkspaceEditorContext(const TSharedRef<IWorkspaceEditor>& InWorkspaceEditor, const FWorkspaceDocument& InDocument) : WorkspaceEditor(InWorkspaceEditor), Document(InDocument)
+{
+}
+
 void FWorkspaceEditorModule::StartupModule()
 {
 	FWorkspaceAssetEditorCommands::Register();
@@ -83,7 +90,7 @@ FObjectDocumentArgs FWorkspaceEditorModule::CreateGraphDocumentArgs(const FGraph
 	Args.OnMakeDocumentWidget = FOnMakeDocumentWidget::CreateLambda([InArgs](const FWorkspaceEditorContext& InContext)
 	{
 		TWeakPtr<FWorkspaceEditor> WeakWorkspaceEditor = StaticCastSharedRef<FWorkspaceEditor>(InContext.WorkspaceEditor);
-		return SNew(SGraphDocument, StaticCastSharedRef<FWorkspaceEditor>(InContext.WorkspaceEditor), CastChecked<UEdGraph>(InContext.Object))
+		return SNew(SGraphDocument, StaticCastSharedRef<FWorkspaceEditor>(InContext.WorkspaceEditor), InContext.Document)
 			.OnCreateActionMenu_Lambda([OnCreateActionMenu = InArgs.OnCreateActionMenu](const FWorkspaceEditorContext& InContext, UEdGraph* InGraph, const FVector2D& InNodePosition, const TArray<UEdGraphPin*>& InDraggedPins, bool bInAutoExpand, SGraphEditor::FActionMenuClosed InOnMenuClosed)
 			{
 				if(OnCreateActionMenu.IsBound())
@@ -114,7 +121,7 @@ FObjectDocumentArgs FWorkspaceEditorModule::CreateGraphDocumentArgs(const FGraph
 	});
 	Args.OnGetTabName = FOnGetTabName::CreateLambda([](const FWorkspaceEditorContext& InContext)
 	{
-		if(UEdGraph* Graph = Cast<UEdGraph>(InContext.Object))
+		if(UEdGraph* Graph = Cast<UEdGraph>(InContext.Document.Object))
 		{
 			if (const UEdGraphSchema* Schema = Graph->GetSchema())
 			{
@@ -142,7 +149,7 @@ FObjectDocumentArgs FWorkspaceEditorModule::CreateGraphDocumentArgs(const FGraph
 			GraphDocument->GraphEditor->GetViewLocation(ViewLocation, ZoomAmount);
 		}
 	
-		return TInstancedStruct<FGraphDocumentState>::Make(InContext.Object, ViewLocation, ZoomAmount);
+		return TInstancedStruct<FGraphDocumentState>::Make(InContext.Document.Object, InContext.Document.Export, ViewLocation, ZoomAmount);
 	});
 	Args.OnSetDocumentState = FOnSetDocumentState::CreateLambda([](const FWorkspaceEditorContext& InContext, TSharedRef<SWidget> InWidget, const TInstancedStruct<FWorkspaceDocumentState>& InDocumentState)
 	{
diff --git a/Engine/Plugins/Experimental/Workspace/Source/WorkspaceEditor/Private/WorkspaceEditorModule.h b/Engine/Plugins/Experimental/Workspace/Source/WorkspaceEditor/Private/WorkspaceEditorModule.h
--- a/Engine/Plugins/Experimental/Workspace/Source/WorkspaceEditor/Private/WorkspaceEditorModule.h
+++ b/Engine/Plugins/Experimental/Workspace/Source/WorkspaceEditor/Private/WorkspaceEditorModule.h
@@ -84,6 +84,7 @@ private:
 	friend class FWorkspaceOutlinerSourceControlColumn;
 	friend class FWorkspaceOutlinerFileStateColumn;
 	friend class FWorkspaceAssetReferenceOutlinerItemDetails;
+	friend struct FWorkspaceDocument;
 	friend SWorkspaceTabWrapper;
 	friend FWorkspaceEditorContext;
 };
diff --git a/Engine/Plugins/Experimental/Workspace/Source/WorkspaceEditor/Public/IWorkspaceEditor.h b/Engine/Plugins/Experimental/Workspace/Source/WorkspaceEditor/Public/IWorkspaceEditor.h
--- a/Engine/Plugins/Experimental/Workspace/Source/WorkspaceEditor/Public/IWorkspaceEditor.h
+++ b/Engine/Plugins/Experimental/Workspace/Source/WorkspaceEditor/Public/IWorkspaceEditor.h
@@ -34,6 +34,9 @@ public:
 	// Open the supplied assets for editing within the workspace editor
 	virtual void OpenAssets(TConstArrayView<FAssetData> InAssets) = 0;
 
+	// Open the supplied exports for editing within the workspace editor
+	virtual void OpenExports(TConstArrayView<FWorkspaceOutlinerItemExport> InExports) = 0;
+
 	// Open the supplied objects for editing within the workspace editor
 	virtual void OpenObjects(TConstArrayView<UObject*> InObjects) = 0;
 
diff --git a/Engine/Plugins/Experimental/Workspace/Source/WorkspaceEditor/Public/IWorkspaceEditorModule.h b/Engine/Plugins/Experimental/Workspace/Source/WorkspaceEditor/Public/IWorkspaceEditorModule.h
--- a/Engine/Plugins/Experimental/Workspace/Source/WorkspaceEditor/Public/IWorkspaceEditorModule.h
+++ b/Engine/Plugins/Experimental/Workspace/Source/WorkspaceEditor/Public/IWorkspaceEditorModule.h
@@ -42,16 +42,37 @@ namespace WorkspaceTabs
 	WORKSPACEEDITOR_API extern const FName BottomRightDocumentArea;
 }
 
+struct WORKSPACEEDITOR_API FWorkspaceDocument
+{
+	FWorkspaceOutlinerItemExport Export;
+	UObject* Object = nullptr;
+	
+	FWorkspaceDocument() = default;
+	FWorkspaceDocument(const FWorkspaceOutlinerItemExport& InExport, UObject* InObject);
+
+	bool operator==(const FWorkspaceDocument& InOther) const
+	{
+		// Comparing object and export path (disregarding export data struct)
+		return Object == InOther.Object && GetTypeHash(Export) == GetTypeHash(InOther.Export);
+	}
+
+	UObject* GetObject() const { check(Object); return Object; }
+	template<class ObjectClass>
+	ObjectClass* GetTypedObject() const { return CastChecked<ObjectClass>(Object); }
+};
+
 // Context passed to workspace editor delegates
-struct FWorkspaceEditorContext
+struct WORKSPACEEDITOR_API FWorkspaceEditorContext
 {
-	FWorkspaceEditorContext(const TSharedRef<IWorkspaceEditor>& InWorkspaceEditor, UObject* InObject);
+	FWorkspaceEditorContext(const TSharedRef<IWorkspaceEditor>& InWorkspaceEditor, const FWorkspaceDocument& InDocument);
+	FWorkspaceEditorContext(const TSharedRef<IWorkspaceEditor>& InWorkspaceEditor, UObject* InObject, const FWorkspaceOutlinerItemExport& InExport);
+	FWorkspaceEditorContext() = delete;
 
 	// The current workspace editor
 	TSharedRef<IWorkspaceEditor> WorkspaceEditor;
 
-	// The object being edited
-	UObject* Object;
+	// The document (object + export) being edited
+	FWorkspaceDocument Document;
 };
 
 struct FWorkspaceBreadcrumb : TSharedFromThis<FWorkspaceBreadcrumb>
diff --git a/Engine/Plugins/Experimental/Workspace/Source/WorkspaceEditor/Public/WorkspaceAssetRegistryInfo.h b/Engine/Plugins/Experimental/Workspace/Source/WorkspaceEditor/Public/WorkspaceAssetRegistryInfo.h
--- a/Engine/Plugins/Experimental/Workspace/Source/WorkspaceEditor/Public/WorkspaceAssetRegistryInfo.h
+++ b/Engine/Plugins/Experimental/Workspace/Source/WorkspaceEditor/Public/WorkspaceAssetRegistryInfo.h
@@ -198,6 +198,27 @@ public:
 			}
 		}
 	}
+	// Returns all valid FWorkspaceOutlinerItemExports found in path segments, starting from the end. e.g:
+	//  - "SoftObjectPath" - "Foo" - "SoftObjectPathTwo" - "Bar" will return "SoftObjectPath" - "Bar" , "SoftObjectPathTwo" - "Foo"
+	void GetExports(TArray<FWorkspaceOutlinerItemExport>& OutExports) const
+	{
+		if(Path.PathSegments.Num() > 0)
+		{
+			bool bFirstExport = true;
+			const int32 NumSegments = Path.PathSegments.Num();
+			for (int32 SegmentIndex = NumSegments - 1; SegmentIndex >= 0; SegmentIndex--)
+			{				
+				FSoftObjectPath ObjectPath = Path.PathSegments[SegmentIndex].ToString();
+				if (ObjectPath.IsValid() && ObjectPath.IsAsset())
+				{
+					FWorkspaceOutlinerItemExport Export;
+					Export.Path.PathSegments.Append(&Path.PathSegments[0], bFirstExport ? NumSegments : SegmentIndex + 1);
+					OutExports.Add(Export);
+					bFirstExport = false;
+				}
+			}
+		}
+	}
 
 	FString GetFullPath() const
 	{
@@ -223,6 +244,10 @@ public:
 	{
 		return GetTypeHash(Export.Path);
 	}
+
+	// Returns the inner ReferredExport from the item data only valid for asset references, otherwise will return *this
+	WORKSPACEEDITOR_API FWorkspaceOutlinerItemExport& GetResolvedExport();	
+	WORKSPACEEDITOR_API const FWorkspaceOutlinerItemExport& GetResolvedExport() const;
 };
 
 USTRUCT()
@@ -240,6 +265,19 @@ struct WORKSPACEEDITOR_API FWorkspaceOutlinerAssetReferenceItemData : public FWo
 
 	UPROPERTY()
 	bool bRecursiveReference = false;
+
+	static bool IsAssetReference(const FWorkspaceOutlinerItemExport& InExport)
+	{
+		if (InExport.HasData())
+		{
+			if (InExport.GetData().GetScriptStruct() == FWorkspaceOutlinerAssetReferenceItemData::StaticStruct())
+			{
+				return true;
+			}			
+		}
+		
+		return false;
+	}
 };
 
 namespace UE::Workspace
diff --git a/Engine/Plugins/Experimental/Workspace/Source/WorkspaceEditor/Public/WorkspaceDocumentState.h b/Engine/Plugins/Experimental/Workspace/Source/WorkspaceEditor/Public/WorkspaceDocumentState.h
--- a/Engine/Plugins/Experimental/Workspace/Source/WorkspaceEditor/Public/WorkspaceDocumentState.h
+++ b/Engine/Plugins/Experimental/Workspace/Source/WorkspaceEditor/Public/WorkspaceDocumentState.h
@@ -2,6 +2,7 @@
 
 #pragma once
 
+#include "WorkspaceAssetRegistryInfo.h"
 #include "WorkspaceDocumentState.generated.h"
 
 // Base struct used to persist workspace document state
@@ -12,15 +13,19 @@ struct FWorkspaceDocumentState
 
 	FWorkspaceDocumentState() = default;
 
-	explicit FWorkspaceDocumentState(const UObject* InObject)
-		: Object(InObject)
+	explicit FWorkspaceDocumentState(const UObject* InObject, const FWorkspaceOutlinerItemExport& InExport)
+		: Object(InObject), Export(InExport)
 	{}
 
 	UPROPERTY()
 	FSoftObjectPath Object;
 
+	// Export provided when opening object
+	UPROPERTY()
+	FWorkspaceOutlinerItemExport Export;
+
 	friend bool operator==(const FWorkspaceDocumentState& LHS, const FWorkspaceDocumentState& RHS)
 	{
-		return LHS.Object == RHS.Object;
+		return LHS.Object == RHS.Object && GetTypeHash(LHS.Export) == GetTypeHash(RHS.Export);
 	}
 };
diff --git a/Engine/Plugins/Experimental/Workspace/Source/WorkspaceEditor/Public/WorkspaceTabPayload.h b/Engine/Plugins/Experimental/Workspace/Source/WorkspaceEditor/Public/WorkspaceTabPayload.h
new file mode 100644
--- /dev/null
+++ b/Engine/Plugins/Experimental/Workspace/Source/WorkspaceEditor/Public/WorkspaceTabPayload.h
@@ -0,0 +1,74 @@
+// Copyright Epic Games, Inc. All Rights Reserved.
+
+#pragma once
+
+#include "CoreMinimal.h"
+#include "WorkspaceAssetRegistryInfo.h"
+#include "WorkflowOrientedApp/WorkflowUObjectDocuments.h"
+
+namespace UE::Workspace
+{
+struct FTabPayload_WorkspaceDocument : FTabPayload
+{
+	static const FName DocumentPayloadName;
+	
+	static TSharedRef<FTabPayload_WorkspaceDocument> Make(const UObject* DocumentID, const FWorkspaceOutlinerItemExport& InExport)
+	{
+		check(DocumentID);
+		return MakeShareable(new FTabPayload_WorkspaceDocument(const_cast<UObject*>(DocumentID), InExport));
+	}
+
+	static TSharedRef<FTabPayload_WorkspaceDocument> Make(const UObject* DocumentID)
+	{
+		check(DocumentID);
+		return MakeShareable(new FTabPayload_WorkspaceDocument(const_cast<UObject*>(DocumentID), FWorkspaceOutlinerItemExport()));
+	}
+
+	template <typename CastType>
+	static CastType* CastChecked(TSharedRef<FTabPayload> Payload)
+	{
+		check(Payload->PayloadType == FTabPayload_WorkspaceDocument::DocumentPayloadName);
+
+		UObject* UntypedObject = StaticCastSharedRef<FTabPayload_WorkspaceDocument>(Payload)->DocumentID.Get(true);
+		return ::CastChecked<CastType>(UntypedObject);
+	}
+
+	static const FWorkspaceOutlinerItemExport& GetExport(const TSharedRef<FTabPayload> Payload)
+	{
+		check(Payload->PayloadType == FTabPayload_WorkspaceDocument::DocumentPayloadName);
+		return StaticCastSharedRef<FTabPayload_WorkspaceDocument>(Payload)->Export;
+	}
+
+	// Determine if another payload is the same as this one
+	virtual bool IsEqual(const TSharedRef<FTabPayload>& OtherPayload) const override
+	{
+		if (OtherPayload->PayloadType == PayloadType)
+		{
+			return this->DocumentID.HasSameIndexAndSerialNumber(StaticCastSharedRef<FTabPayload_WorkspaceDocument>(OtherPayload)->DocumentID)
+			// Replace with ==operator
+			&& GetTypeHash(StaticCastSharedRef<FTabPayload_WorkspaceDocument>(OtherPayload)->Export) == GetTypeHash(Export);
+		}
+
+		return false;
+	}
+
+	virtual bool IsValid() const override
+	{
+		return DocumentID.IsValid() && (Export.GetFirstAssetPath().IsValid() || Export.GetFirstAssetPath().IsNull());
+	}
+
+	virtual ~FTabPayload_WorkspaceDocument() override {};
+	
+private:
+	FTabPayload_WorkspaceDocument(UObject* InDocumentID, const FWorkspaceOutlinerItemExport& InExport)
+		: FTabPayload(DocumentPayloadName)
+		, DocumentID(InDocumentID)
+		, Export(InExport)
+	{
+	}
+
+private:
+	TWeakObjectPtr<UObject> DocumentID;
+	FWorkspaceOutlinerItemExport Export;
+};
+}
\ No newline at end of file
