Multiple Anim Next tests failing on Windows & Linux due to assertion
#fix ensure valid in-place workspace export is generated when opening an asset by UObject-ptr rather than by its export

diff --git a/Engine/Plugins/Experimental/Workspace/Source/WorkspaceEditor/Private/WorkspaceEditor.cpp b/Engine/Plugins/Experimental/Workspace/Source/WorkspaceEditor/Private/WorkspaceEditor.cpp
--- a/Engine/Plugins/Experimental/Workspace/Source/WorkspaceEditor/Private/WorkspaceEditor.cpp
+++ b/Engine/Plugins/Experimental/Workspace/Source/WorkspaceEditor/Private/WorkspaceEditor.cpp
@@ -284,7 +284,7 @@ void FWorkspaceEditor::SaveEditedObjectState() const
 
 TSharedPtr<SDockTab> FWorkspaceEditor::OpenDocument(const UObject* InForObject, FDocumentTracker::EOpenDocumentCause InCause)
 {
-	FWorkspaceOutlinerItemExport Export(InForObject->GetFName(), InForObject->GetPackage());
+	FWorkspaceOutlinerItemExport Export(InForObject->GetFName(), InForObject);
 	return OpenDocument(InForObject, Export, InCause);
 }
 
