Added GetFirstObjectOfType and use that instead for breadcrumb population (graph functions are sub-objects not assets/objects)

diff --git a/Engine/Plugins/Experimental/Workspace/Source/WorkspaceEditor/Private/SWorkspaceTabWrapper.cpp b/Engine/Plugins/Experimental/Workspace/Source/WorkspaceEditor/Private/SWorkspaceTabWrapper.cpp
--- a/Engine/Plugins/Experimental/Workspace/Source/WorkspaceEditor/Private/SWorkspaceTabWrapper.cpp
+++ b/Engine/Plugins/Experimental/Workspace/Source/WorkspaceEditor/Private/SWorkspaceTabWrapper.cpp
@@ -243,7 +243,7 @@ void SWorkspaceTabWrapper::RebuildBreadcrumbTrail() const
 
 		for (const FWorkspaceOutlinerItemExport& DocumentExport : Exports)
 		{
-			UObject* DocumentID = DocumentExport.GetFirstAssetOfType<UObject>();
+			UObject* DocumentID = DocumentExport.GetFirstObjectOfType<UObject>();
 			check(DocumentID);
 			if (const TSharedPtr<UE::Workspace::FWorkspaceEditor> SharedWorkspaceEditor = WeakWorkspaceEditor.Pin())
 			{
diff --git a/Engine/Plugins/Experimental/Workspace/Source/WorkspaceEditor/Public/WorkspaceAssetRegistryInfo.h b/Engine/Plugins/Experimental/Workspace/Source/WorkspaceEditor/Public/WorkspaceAssetRegistryInfo.h
--- a/Engine/Plugins/Experimental/Workspace/Source/WorkspaceEditor/Public/WorkspaceAssetRegistryInfo.h
+++ b/Engine/Plugins/Experimental/Workspace/Source/WorkspaceEditor/Public/WorkspaceAssetRegistryInfo.h
@@ -146,6 +146,30 @@ public:
 		return nullptr;
 	}
 
+	template<typename AssetClass>
+	AssetClass* GetFirstObjectOfType() const
+	{
+		if(Path.PathSegments.Num() > 0)
+		{
+			for (int32 SegmentIndex = Path.PathSegments.Num() - 1; SegmentIndex >= 0; SegmentIndex--)
+			{
+				FSoftObjectPath ObjectPath = Path.PathSegments[SegmentIndex].ToString();
+				if (ObjectPath.IsValid())
+				{
+					if (UObject* Object = ObjectPath.TryLoad())
+					{
+						if (AssetClass* TypedObject = Cast<AssetClass>(Object))
+						{
+							return TypedObject;
+						}
+					}
+				}
+			}
+		}
+
+		return nullptr;
+	}
+
 	// Returns the first FSoftObjectPath found in segments, starting from the end. e.g. "SoftObjectPath" - "Foo" - "SoftObjectPathTwo" - "Bar" will return SoftObjectPathTwo
 	FSoftObjectPath GetFirstAssetPath() const
 	{
