AnimNext variables UI/UX changes:

diff --git a/Engine/Plugins/Experimental/Workspace/Source/WorkspaceEditor/Private/Outliner/WorkspaceOutlinerTreeItem.cpp b/Engine/Plugins/Experimental/Workspace/Source/WorkspaceEditor/Private/Outliner/WorkspaceOutlinerTreeItem.cpp
--- a/Engine/Plugins/Experimental/Workspace/Source/WorkspaceEditor/Private/Outliner/WorkspaceOutlinerTreeItem.cpp
+++ b/Engine/Plugins/Experimental/Workspace/Source/WorkspaceEditor/Private/Outliner/WorkspaceOutlinerTreeItem.cpp
@@ -12,7 +12,6 @@
 #include "Widgets/DeclarativeSyntaxSupport.h"
 #include "Widgets/Images/SImage.h"
 #include "Widgets/Text/SInlineEditableTextBlock.h"
-#include "Widgets/Text/STextBlock.h"
 #include "Widgets/Views/STableRow.h"
 #include "Widgets/Views/SListView.h"
 
@@ -21,9 +20,9 @@
 namespace UE::Workspace
 {
 	const FSceneOutlinerTreeItemType FWorkspaceOutlinerTreeItem::Type;
-	class SWorkspaceOutlinerTreelabel : FSceneOutlinerCommonLabelData, public SCompoundWidget
+	class SWorkspaceOutlinerTreeLabel : FSceneOutlinerCommonLabelData, public SCompoundWidget
 	{
-		SLATE_BEGIN_ARGS(SWorkspaceOutlinerTreelabel) {}
+		SLATE_BEGIN_ARGS(SWorkspaceOutlinerTreeLabel) {}
 		SLATE_END_ARGS()
 
 		void Construct(const FArguments& InArgs, FWorkspaceOutlinerTreeItem& InTreeItem, ISceneOutliner& SceneOutliner, const STableRow<FSceneOutlinerTreeItemPtr>& InRow)
@@ -76,11 +75,12 @@ namespace UE::Workspace
 				.Padding(0.0f, 2.0f)
 				[
 					SAssignNew(TextBlock, SInlineEditableTextBlock)
-					.Text(this, &SWorkspaceOutlinerTreelabel::GetDisplayText)
+					.Text(this, &SWorkspaceOutlinerTreeLabel::GetDisplayText)
 					.HighlightText(SceneOutliner.GetFilterHighlightText())
-					.ColorAndOpacity(this, &SWorkspaceOutlinerTreelabel::GetForegroundColor)
-					.OnTextCommitted(this, &SWorkspaceOutlinerTreelabel::OnTextCommited)
-					.OnVerifyTextChanged(this, &SWorkspaceOutlinerTreelabel::OnVerifyTextChanged)
+					.ColorAndOpacity(this, &SWorkspaceOutlinerTreeLabel::GetForegroundColor)
+					.OnTextCommitted(this, &SWorkspaceOutlinerTreeLabel::OnTextCommited)
+					.OnVerifyTextChanged(this, &SWorkspaceOutlinerTreeLabel::OnVerifyTextChanged)
+					.ToolTipText(this, &SWorkspaceOutlinerTreeLabel::GetToolTipText)
 				]
 			];
 		}
@@ -94,7 +94,16 @@ namespace UE::Workspace
 			return FText();
 		}
 
-		void OnTextCommited(const FText& InLabel, ETextCommit::Type InCommitInfo)
+		FText GetToolTipText() const
+		{
+			if (const TSharedPtr<FWorkspaceOutlinerTreeItem> Item = TreeItem.Pin())
+			{
+				return Item->GetToolTipText();
+			}
+			return FText();
+		}
+
+		void OnTextCommited(const FText& InLabel, ETextCommit::Type InCommitInfo) const
 		{
 			if(InCommitInfo == ETextCommit::OnEnter)
 			{
@@ -108,7 +117,7 @@ namespace UE::Workspace
 			}
 		}
 
-		bool OnVerifyTextChanged(const FText& InLabel, FText& OutErrorMessage)
+		bool OnVerifyTextChanged(const FText& InLabel, FText& OutErrorMessage) const
 		{
 			if (const TSharedPtr<FWorkspaceOutlinerTreeItem> Item = TreeItem.Pin())
 			{
@@ -159,9 +168,19 @@ namespace UE::Workspace
 		return Export.GetIdentifier().ToString();
 	}
 
+	FText FWorkspaceOutlinerTreeItem::GetToolTipText() const
+	{
+		if (const TSharedPtr<IWorkspaceOutlinerItemDetails> SharedFactory = FWorkspaceEditorModule::GetOutlinerItemDetails(MakeOutlinerDetailsId(Export)))
+		{
+			return SharedFactory->GetToolTipText(Export);
+		}
+		
+		return FText::GetEmpty();
+	}
+
 	TSharedRef<SWidget> FWorkspaceOutlinerTreeItem::GenerateLabelWidget(ISceneOutliner& Outliner, const STableRow<FSceneOutlinerTreeItemPtr>& InRow)
 	{
-		TSharedRef<SWorkspaceOutlinerTreelabel> LabelWidget = SNew(SWorkspaceOutlinerTreelabel, *this, Outliner, InRow);
+		TSharedRef<SWorkspaceOutlinerTreeLabel> LabelWidget = SNew(SWorkspaceOutlinerTreeLabel, *this, Outliner, InRow);
 		RenameRequestEvent.BindSP(LabelWidget->TextBlock.Get(), &SInlineEditableTextBlock::EnterEditingMode);
 		return LabelWidget;
 	}
diff --git a/Engine/Plugins/Experimental/Workspace/Source/WorkspaceEditor/Private/Outliner/WorkspaceOutlinerTreeItem.h b/Engine/Plugins/Experimental/Workspace/Source/WorkspaceEditor/Private/Outliner/WorkspaceOutlinerTreeItem.h
--- a/Engine/Plugins/Experimental/Workspace/Source/WorkspaceEditor/Private/Outliner/WorkspaceOutlinerTreeItem.h
+++ b/Engine/Plugins/Experimental/Workspace/Source/WorkspaceEditor/Private/Outliner/WorkspaceOutlinerTreeItem.h
@@ -23,6 +23,7 @@ namespace UE::Workspace
 		virtual bool IsValid() const override;
 		virtual FSceneOutlinerTreeItemID GetID() const override;
 		virtual FString GetDisplayString() const override;
+		virtual FText GetToolTipText() const;
 		virtual bool CanInteract() const override { return true; }
 		virtual TSharedRef<SWidget> GenerateLabelWidget(ISceneOutliner& Outliner, const STableRow<FSceneOutlinerTreeItemPtr>& InRow) override;
 		virtual FString GetPackageName() const override;
diff --git a/Engine/Plugins/Experimental/Workspace/Source/WorkspaceEditor/Private/SGraphDocument.cpp b/Engine/Plugins/Experimental/Workspace/Source/WorkspaceEditor/Private/SGraphDocument.cpp
--- a/Engine/Plugins/Experimental/Workspace/Source/WorkspaceEditor/Private/SGraphDocument.cpp
+++ b/Engine/Plugins/Experimental/Workspace/Source/WorkspaceEditor/Private/SGraphDocument.cpp
@@ -13,7 +13,7 @@ namespace UE::Workspace
 void SGraphDocument::Construct(const FArguments& InArgs, TSharedRef<FWorkspaceEditor> InHostingApp, const FWorkspaceDocument& InWorkspaceDocument)
 {
 	HostingAppPtr = InHostingApp;
-	EdGraph = InWorkspaceDocument.GetTypedObject<UEdGraph>();
+	EdGraph = InWorkspaceDocument.GetTypedObjectChecked<UEdGraph>();
 	Document = InWorkspaceDocument;
 
 	BindCommands();
diff --git a/Engine/Plugins/Experimental/Workspace/Source/WorkspaceEditor/Private/SWorkspaceTabWrapper.cpp b/Engine/Plugins/Experimental/Workspace/Source/WorkspaceEditor/Private/SWorkspaceTabWrapper.cpp
--- a/Engine/Plugins/Experimental/Workspace/Source/WorkspaceEditor/Private/SWorkspaceTabWrapper.cpp
+++ b/Engine/Plugins/Experimental/Workspace/Source/WorkspaceEditor/Private/SWorkspaceTabWrapper.cpp
@@ -14,7 +14,7 @@ void SWorkspaceTabWrapper::Construct( const FArguments& InArgs, TSharedPtr<class
 {
 	Content = InArgs._Content.Widget;
 	WeakWorkspaceEditor = StaticCastSharedRef<UE::Workspace::FWorkspaceEditor>(InEditorContext.WorkspaceEditor);
-	WeakDocumentObject = InEditorContext.Document.Object;
+	WeakDocumentObject = InEditorContext.Document.GetObject();
 	Export = InEditorContext.Document.Export;
 
 	if (InTabInfo->GetTab().IsValid())
@@ -225,10 +225,10 @@ void SWorkspaceTabWrapper::RebuildBreadcrumbTrail() const
 	BreadcrumbTrail->ClearCrumbs(false);
 
 	TArray<TSharedPtr<UE::Workspace::FWorkspaceBreadcrumb>> Breadcrumbs;
-	TArray<FWorkspaceOutlinerItemExport> Exports;
 
 	if (Export.GetFirstAssetPath().IsValid())
 	{
+		TArray<FWorkspaceOutlinerItemExport> Exports;
 		if (FWorkspaceOutlinerAssetReferenceItemData::IsAssetReference(Export))
 		{
 			// Prepend outer asset/export chain
@@ -245,14 +245,19 @@ void SWorkspaceTabWrapper::RebuildBreadcrumbTrail() const
 		{
 			UObject* DocumentID = DocumentExport.GetFirstObjectOfType<UObject>();
 			check(DocumentID);
-			if (const TSharedPtr<UE::Workspace::FWorkspaceEditor> SharedWorkspaceEditor = WeakWorkspaceEditor.Pin())
+
+			// [TODO] DocumentID can fail to resolve for previously moved assets (need to resolve redirector?)
+			if (DocumentID)
 			{
-				const UE::Workspace::FWorkspaceEditorModule& WorkspaceEditorModule = FModuleManager::LoadModuleChecked<UE::Workspace::FWorkspaceEditorModule>("WorkspaceEditor");
-				if(const UE::Workspace::FObjectDocumentArgs* DocumentArgs = WorkspaceEditorModule.FindObjectDocumentType(DocumentID))
+				if (const TSharedPtr<UE::Workspace::FWorkspaceEditor> SharedWorkspaceEditor = WeakWorkspaceEditor.Pin())
 				{
-					if(DocumentArgs->OnGetDocumentBreadcrumbTrail.IsBound())
+					const UE::Workspace::FWorkspaceEditorModule& WorkspaceEditorModule = FModuleManager::LoadModuleChecked<UE::Workspace::FWorkspaceEditorModule>("WorkspaceEditor");
+					if(const UE::Workspace::FObjectDocumentArgs* DocumentArgs = WorkspaceEditorModule.FindObjectDocumentType({DocumentExport, DocumentID}))
 					{
-						DocumentArgs->OnGetDocumentBreadcrumbTrail.Execute(UE::Workspace::FWorkspaceEditorContext(SharedWorkspaceEditor.ToSharedRef(), { DocumentExport, DocumentID}), Breadcrumbs);
+						if(DocumentArgs->OnGetDocumentBreadcrumbTrail.IsBound())
+						{
+							DocumentArgs->OnGetDocumentBreadcrumbTrail.Execute(UE::Workspace::FWorkspaceEditorContext(SharedWorkspaceEditor.ToSharedRef(), { DocumentExport, DocumentID}), Breadcrumbs);
+						}
 					}
 				}
 			}
@@ -265,7 +270,7 @@ void SWorkspaceTabWrapper::RebuildBreadcrumbTrail() const
 		    if (const TSharedPtr<UE::Workspace::FWorkspaceEditor> SharedWorkspaceEditor = WeakWorkspaceEditor.Pin())
 		    {
 			    const UE::Workspace::FWorkspaceEditorModule& WorkspaceEditorModule = FModuleManager::LoadModuleChecked<UE::Workspace::FWorkspaceEditorModule>("WorkspaceEditor");
-			    if(const UE::Workspace::FObjectDocumentArgs* DocumentArgs = WorkspaceEditorModule.FindObjectDocumentType(DocumentID))
+			    if(const UE::Workspace::FObjectDocumentArgs* DocumentArgs = WorkspaceEditorModule.FindObjectDocumentType({Export, DocumentID}))
 			    {
 				    if(DocumentArgs->OnGetDocumentBreadcrumbTrail.IsBound())
 				    {
@@ -291,7 +296,7 @@ const FSlateBrush* SWorkspaceTabWrapper::GetTabIcon() const
 		if (const TSharedPtr<UE::Workspace::FWorkspaceEditor> SharedWorkspaceEditor = WeakWorkspaceEditor.Pin())
 		{
 			const UE::Workspace::FWorkspaceEditorModule& WorkspaceEditorModule = FModuleManager::LoadModuleChecked<UE::Workspace::FWorkspaceEditorModule>("WorkspaceEditor");
-			if(const UE::Workspace::FObjectDocumentArgs* DocumentArgs = WorkspaceEditorModule.FindObjectDocumentType(DocumentID))
+			if(const UE::Workspace::FObjectDocumentArgs* DocumentArgs = WorkspaceEditorModule.FindObjectDocumentType({Export, DocumentID}))
 			{
 				if(DocumentArgs->OnGetTabIcon.IsBound())
 				{
diff --git a/Engine/Plugins/Experimental/Workspace/Source/WorkspaceEditor/Private/SWorkspaceTabWrapper.h b/Engine/Plugins/Experimental/Workspace/Source/WorkspaceEditor/Private/SWorkspaceTabWrapper.h
--- a/Engine/Plugins/Experimental/Workspace/Source/WorkspaceEditor/Private/SWorkspaceTabWrapper.h
+++ b/Engine/Plugins/Experimental/Workspace/Source/WorkspaceEditor/Private/SWorkspaceTabWrapper.h
@@ -21,6 +21,7 @@ public:
 	void Construct(const FArguments& InArgs, TSharedPtr<class FTabInfo> InTabInfo, const UE::Workspace::FWorkspaceEditorContext& InEditorContext);
 	TSharedRef<SWidget> GetContent() const { return Content.ToSharedRef(); }
 	TWeakObjectPtr<UObject> GetDocumentObject() const { return WeakDocumentObject; }
+	UE::Workspace::FWorkspaceDocument GetDocument() const { return { Export, WeakDocumentObject.Get() }; }
 	
 protected:
 	void RebuildBreadcrumbTrail() const;
diff --git a/Engine/Plugins/Experimental/Workspace/Source/WorkspaceEditor/Private/SWorkspaceView.cpp b/Engine/Plugins/Experimental/Workspace/Source/WorkspaceEditor/Private/SWorkspaceView.cpp
--- a/Engine/Plugins/Experimental/Workspace/Source/WorkspaceEditor/Private/SWorkspaceView.cpp
+++ b/Engine/Plugins/Experimental/Workspace/Source/WorkspaceEditor/Private/SWorkspaceView.cpp
@@ -11,9 +11,8 @@
 #include "SPositiveActionButton.h"
 #include "SSceneOutliner.h"
 #include "WorkspaceAssetRegistryInfo.h"
-#include "WorkspaceSchema.h"
-#include "Framework/Commands/GenericCommands.h"
 #include "Outliner/WorkspaceOutlinerColumns.h"
+#include "Outliner/WorkspaceOutlinerHierarchy.h"
 #include "Outliner/WorkspaceOutlinerMode.h"
 #include "Outliner/WorkspaceOutlinerTreeItem.h"
 
@@ -38,9 +37,9 @@ public:
 	}
 
 	SWorkspaceOutliner() = default;
-	virtual ~SWorkspaceOutliner() = default;
+	virtual ~SWorkspaceOutliner() override = default;
 
-	void CustomAddToToolbar(TSharedPtr<class SHorizontalBox> Toolbar) override
+	virtual void CustomAddToToolbar(TSharedPtr<class SHorizontalBox> Toolbar) override
 	{
 		Toolbar->AddSlot()
 		.VAlign(VAlign_Center)
@@ -103,7 +102,11 @@ void SWorkspaceView::Construct(const FArguments& InArgs, UWorkspace* InWorkspace
 		InitOptions.ColumnMap.Add(FSceneOutlinerBuiltInColumnTypes::Label(), FSceneOutlinerColumnInfo(ESceneOutlinerColumnVisibility::Visible, 10));
 		InitOptions.ColumnMap.Add(FWorkspaceOutlinerFileStateColumn::GetID(), FSceneOutlinerColumnInfo(ESceneOutlinerColumnVisibility::Visible, 0, FCreateSceneOutlinerColumn::CreateLambda([](ISceneOutliner& InSceneOutliner) { return MakeShareable(new FWorkspaceOutlinerFileStateColumn(InSceneOutliner)); }), false));
 		InitOptions.ColumnMap.Add(FSceneOutlinerSourceControlColumn::GetID(), FSceneOutlinerColumnInfo(ESceneOutlinerColumnVisibility::Visible, 100, FCreateSceneOutlinerColumn::CreateLambda([](ISceneOutliner& InSceneOutliner) { return MakeShareable(new FSceneOutlinerSourceControlColumn(InSceneOutliner)); }), false));	
-		InitOptions.ModeFactory = FCreateSceneOutlinerMode::CreateLambda([this, WeakWorkspaceEditor=InWorkspaceEditor.ToWeakPtr()](SSceneOutliner* InOutliner) { return new UE::Workspace::FWorkspaceOutlinerMode(UE::Workspace::FWorkspaceOutlinerMode(InOutliner, Workspace, WeakWorkspaceEditor)); });
+		InitOptions.ModeFactory = FCreateSceneOutlinerMode::CreateLambda([this, WeakWorkspaceEditor=InWorkspaceEditor.ToWeakPtr()](SSceneOutliner* InOutliner)
+		{
+			OutlinerMode = new UE::Workspace::FWorkspaceOutlinerMode(UE::Workspace::FWorkspaceOutlinerMode(InOutliner, Workspace, WeakWorkspaceEditor));
+			return OutlinerMode;
+		});
 	}
 	SceneWorkspaceOutliner = SNew(SWorkspaceOutliner, InitOptions, Workspace);
 
@@ -169,7 +172,7 @@ void SWorkspaceView::Construct(const FArguments& InArgs, UWorkspace* InWorkspace
 	];
 }
 
-void SWorkspaceView::SelectObject(UObject* InObject) const
+void SWorkspaceView::SelectObject(const UObject* InObject) const
 {
 	FWorkspaceOutlinerItemExport Export(InObject->GetFName(), InObject);
 
@@ -200,6 +203,17 @@ void SWorkspaceView::GetWorkspaceExportData(FWorkspaceOutlinerItemExport& InOutP
 		}
 	}
 }
+
+void SWorkspaceView::GetHierarchyAssetData(TArray<FAssetData>& OutAssetData) const
+{
+	if (OutlinerMode)
+	{
+		if (FWorkspaceOutlinerHierarchy* Hierarchy = static_cast<FWorkspaceOutlinerHierarchy*>(OutlinerMode->GetHierarchy()))
+		{
+			OutAssetData.Append(Hierarchy->ProcessedAssetData);
+		}
+	}	
+}
 }
 
 #undef LOCTEXT_NAMESPACE
\ No newline at end of file
diff --git a/Engine/Plugins/Experimental/Workspace/Source/WorkspaceEditor/Private/SWorkspaceView.h b/Engine/Plugins/Experimental/Workspace/Source/WorkspaceEditor/Private/SWorkspaceView.h
--- a/Engine/Plugins/Experimental/Workspace/Source/WorkspaceEditor/Private/SWorkspaceView.h
+++ b/Engine/Plugins/Experimental/Workspace/Source/WorkspaceEditor/Private/SWorkspaceView.h
@@ -13,6 +13,7 @@ namespace UE::Workspace
 {
 class SWorkspaceOutliner;
 class IWorkspaceEditor;
+class FWorkspaceOutlinerMode;
 
 class SWorkspaceView : public SCompoundWidget
 {
@@ -21,13 +22,15 @@ class SWorkspaceView : public SCompoundWidget
 
 	void Construct(const FArguments& InArgs, UWorkspace* InWorkspace, TSharedRef<UE::Workspace::IWorkspaceEditor> InWorkspaceEditor);
 
-	void SelectObject(UObject* InObject) const;
+	void SelectObject(const UObject* InObject) const;
 	void SelectExport(const FWorkspaceOutlinerItemExport& InExport) const;
 	// Tries to lookup the FWorkspaceOutlinerItemExport path within the outliner, and copies over its TInstancedStruct<FWorkspaceOutlinerItemData> value
 	void GetWorkspaceExportData(FWorkspaceOutlinerItemExport& InOutPartialExport) const;
+	void GetHierarchyAssetData(TArray<FAssetData>& OutAssetData) const;
 	
 private:
 	UWorkspace* Workspace = nullptr;
+	FWorkspaceOutlinerMode* OutlinerMode = nullptr;
 	TSharedPtr<SWorkspaceOutliner> SceneWorkspaceOutliner;
 };
 
diff --git a/Engine/Plugins/Experimental/Workspace/Source/WorkspaceEditor/Private/WorkspaceAssetReferenceItemDetails.h b/Engine/Plugins/Experimental/Workspace/Source/WorkspaceEditor/Private/WorkspaceAssetReferenceItemDetails.h
--- a/Engine/Plugins/Experimental/Workspace/Source/WorkspaceEditor/Private/WorkspaceAssetReferenceItemDetails.h
+++ b/Engine/Plugins/Experimental/Workspace/Source/WorkspaceEditor/Private/WorkspaceAssetReferenceItemDetails.h
@@ -43,8 +43,7 @@ public:
 			return GetInnerDetails(Export)->GetDisplayString(GetInnerExport(Export));
 		}
 
-		return bShowFullIdentifierInOutliner ? FString::Printf(TEXT("%s - %s - %s"), *Export.GetIdentifier().ToString(), *(Export.GetData().GetScriptStruct() ? Export.GetData().GetScriptStruct()->GetFName() : NAME_None).ToString(), *Export.GetFirstAssetPath().ToString()) :
-			Export.GetFirstAssetPath().GetAssetName();
+		return Export.GetFirstAssetPath().GetAssetName();
 	}
 	
 	virtual const FSlateBrush* GetItemIcon(const FWorkspaceOutlinerItemExport& Export) const override
@@ -64,7 +63,7 @@ public:
 	virtual FSlateColor GetItemColor(const FWorkspaceOutlinerItemExport& Export) const override
 	{
 		const FWorkspaceOutlinerAssetReferenceItemData& Data = Export.GetData().Get<FWorkspaceOutlinerAssetReferenceItemData>();	
-		return Data.bRecursiveReference ? FLinearColor::Red : FSlateColor::UseSubduedForeground();
+		return Data.bRecursiveReference && Data.bShouldExpandReference ? FLinearColor::Red : FSlateColor::UseSubduedForeground();
 	}	
 
 	virtual bool IsExpandedByDefault() const override
diff --git a/Engine/Plugins/Experimental/Workspace/Source/WorkspaceEditor/Private/WorkspaceEditor.cpp b/Engine/Plugins/Experimental/Workspace/Source/WorkspaceEditor/Private/WorkspaceEditor.cpp
--- a/Engine/Plugins/Experimental/Workspace/Source/WorkspaceEditor/Private/WorkspaceEditor.cpp
+++ b/Engine/Plugins/Experimental/Workspace/Source/WorkspaceEditor/Private/WorkspaceEditor.cpp
@@ -6,7 +6,6 @@
 #include "ExternalPackageHelper.h"
 #include "IStructureDetailsView.h"
 #include "Widgets/Docking/SDockTab.h"
-#include "WorkflowOrientedApp/WorkflowUObjectDocuments.h"
 #include "WorkspaceAssetEditor.h"
 #include "WorkspaceState.h"
 #include "WorkspaceDocumentState.h"
@@ -17,17 +16,14 @@
 #include "WorkspaceEditorCommands.h"
 #include "SWorkspaceTabWrapper.h"
 #include "Dialogs/Dialogs.h"
-#include "Framework/Docking/LayoutExtender.h"
 #include "ToolMenus.h"
 #include "AssetEditorModeManager.h"
 #include "Selection.h"
 #include "EditorModeManager.h"
 #include "Tools/EdModeInteractiveToolsContext.h"
-#include "ContextObjectStore.h"
 #include "SGraphDocument.h"
 #include "WorkspaceSchema.h"
 #include "WorkspaceTabPayload.h"
-#include "SWorkspaceView.h"
 
 #define LOCTEXT_NAMESPACE "WorkspaceEditor"
 
@@ -258,7 +254,7 @@ void FWorkspaceEditor::RestoreEditedObjectState()
 			{
 				const TSharedRef<SWorkspaceTabWrapper> TabWrapper = StaticCastSharedRef<SWorkspaceTabWrapper>(DockTab->GetContent());	
 				FWorkspaceEditorModule& WorkspaceEditorModule = FModuleManager::LoadModuleChecked<FWorkspaceEditorModule>("WorkspaceEditor");
-				const FObjectDocumentArgs* DocumentArgs = WorkspaceEditorModule.FindObjectDocumentType(Object);
+				const FObjectDocumentArgs* DocumentArgs = WorkspaceEditorModule.FindObjectDocumentType({Export, Object});
 				if(DocumentArgs != nullptr && DocumentArgs->OnSetDocumentState.IsBound())
 				{
 					DocumentArgs->OnSetDocumentState.Execute(FWorkspaceEditorContext(SharedThis(this), { Export, Object }), TabWrapper->GetContent(), DocumentState);
@@ -302,7 +298,7 @@ TSharedPtr<SDockTab> FWorkspaceEditor::OpenDocument(const UObject* InForObject,
 	if(InCause != FDocumentTracker::RestorePreviousDocument)
 	{
 		const FWorkspaceEditorModule& WorkspaceEditorModule = FModuleManager::LoadModuleChecked<FWorkspaceEditorModule>("WorkspaceEditor");
-		const FObjectDocumentArgs* DocumentArgs = WorkspaceEditorModule.FindObjectDocumentType(InForObject);
+		const FObjectDocumentArgs* DocumentArgs = WorkspaceEditorModule.FindObjectDocumentType({InExport, const_cast<UObject*>(InForObject)});
 		if(DocumentArgs != nullptr && DocumentArgs->OnGetDocumentState.IsBound())
 		{
 			const TSharedRef<SWorkspaceTabWrapper> TabWrapper = StaticCastSharedRef<SWorkspaceTabWrapper>(NewTab->GetContent());
@@ -407,14 +403,21 @@ void FWorkspaceEditor::GetOpenedAssetsOfClass(TSubclassOf<UObject> InClass, TArr
 	}
 }
 
-void FWorkspaceEditor::GetAssets(TArray<FAssetData>& OutAssets) const
+void FWorkspaceEditor::GetAssets(TArray<FAssetData>& OutAssets, bool bIncludeAssetReferences /*= false */) const
 {
-	FWorkspaceAssetRegistryExports Exports;
-	FWorkspaceEditorModule::GetExportedAssetsForWorkspace(FAssetData(Workspace), Exports);
-	IAssetRegistry& AssetRegistry = FModuleManager::LoadModuleChecked<FAssetRegistryModule>("AssetRegistry").Get();
-	for(const FWorkspaceAssetRegistryExportEntry& Entry : Exports.Assets)
+	if (bIncludeAssetReferences)
 	{
-		OutAssets.Add(AssetRegistry.GetAssetByObjectPath(Entry.Asset));
+		WorkspaceView->GetHierarchyAssetData(OutAssets);
+	}
+	else
+	{
+		FWorkspaceAssetRegistryExports Exports;
+		FWorkspaceEditorModule::GetExportedAssetsForWorkspace(FAssetData(Workspace), Exports);
+		IAssetRegistry& AssetRegistry = FModuleManager::LoadModuleChecked<FAssetRegistryModule>("AssetRegistry").Get();
+		for(const FWorkspaceAssetRegistryExportEntry& Entry : Exports.Assets)
+		{
+			OutAssets.Add(AssetRegistry.GetAssetByObjectPath(Entry.Asset));
+		}
 	}
 }
 
@@ -483,23 +486,25 @@ void FWorkspaceEditor::SetGlobalSelection(FGlobalSelectionId SelectionId, FOnCle
 	GlobalSelections.Add({ SelectionId, OnClearSelectionDelegate});
 }
 
-void FWorkspaceEditor::SetFocussedDocument(const TObjectPtr<UObject> InDocument)
+void FWorkspaceEditor::SetFocusedDocument(const UE::Workspace::FWorkspaceDocument& Document)
 {
-	if (bSettingFocussedDocument)
+	if (bSettingFocusedDocument)
 	{
 		return;
 	}	
-	TGuardValue<bool> SetValueGuard(bSettingFocussedDocument, true);
+	TGuardValue<bool> SetValueGuard(bSettingFocusedDocument, true);
 	
 	FAssetEditorModeManager* ModeManager = static_cast<FAssetEditorModeManager*>(&GetEditorModeManager());	
 	ModeManager->GetSelectedObjects()->DeselectAll();
 
-	if (InDocument != nullptr)
+	UObject* DocumentObject = Document.GetObject();
+
+	if (DocumentObject != nullptr)
 	{
-		ModeManager->GetSelectedObjects()->Select(InDocument);
+		ModeManager->GetSelectedObjects()->Select(DocumentObject);
 
 		FWorkspaceEditorModule& WorkspaceEditorModule = FModuleManager::Get().LoadModuleChecked<FWorkspaceEditorModule>("WorkspaceEditor");
-		if (const FObjectDocumentArgs* Args = WorkspaceEditorModule.FindObjectDocumentType(InDocument))
+		if (const FObjectDocumentArgs* Args = WorkspaceEditorModule.FindObjectDocumentType(Document))
 		{
 			const FEditorModeID NewEditorMode = Args->DocumentEditorMode;
 			if (NewEditorMode != NAME_None)
@@ -514,24 +519,30 @@ void FWorkspaceEditor::SetFocussedDocument(const TObjectPtr<UObject> InDocument)
 				ModeManager->DeactivateAllModes();
 			}
 		}
+
+		LastDocument = Document;
 	}
 	else
 	{
 		ModeManager->DeactivateAllModes();
+		LastDocument.Reset();
 	}
 
-	OnFocussedDocumentChangedDelegate.Broadcast(InDocument);
+	OnFocusedDocumentChangedDelegate.Broadcast(Document);
 }
 
-const TObjectPtr<UObject> FWorkspaceEditor::GetFocussedDocumentOfClass(const TObjectPtr<UClass> AssetClass) const
+const FWorkspaceDocument& FWorkspaceEditor::GetFocusedDocumentOfClass(const TObjectPtr<UClass> AssetClass) const
 {
-	FAssetEditorModeManager* ModeManager = static_cast<FAssetEditorModeManager*>(&GetEditorModeManager());	
-	if (USelection* Selection = ModeManager->GetSelectedObjects())
+	if (LastDocument.IsSet())
 	{
-		return Selection->GetTop(AssetClass);
+		if (LastDocument.GetValue().GetObject() && LastDocument.GetValue().GetObject()->IsA(AssetClass))
+		{
+			return LastDocument.GetValue();
+		}
 	}
 
-	return nullptr;
+	static FWorkspaceDocument EmptyDocument;
+	return EmptyDocument;
 }
 
 void FWorkspaceEditor::HandleOutlinerSelectionChanged(TConstArrayView<FWorkspaceOutlinerItemExport> InExports)
@@ -539,9 +550,9 @@ void FWorkspaceEditor::HandleOutlinerSelectionChanged(TConstArrayView<FWorkspace
 	LastSelectedExports = InExports;
 }
 
-FOnFocussedDocumentChanged& FWorkspaceEditor::OnFocussedDocumentChanged()
+FOnFocusedDocumentChanged& FWorkspaceEditor::OnFocusedDocumentChanged()
 {
-	return OnFocussedDocumentChangedDelegate;
+	return OnFocusedDocumentChangedDelegate;
 }
 
 TSharedPtr<IDetailsView> FWorkspaceEditor::GetDetailsView()
@@ -559,6 +570,17 @@ FString FWorkspaceEditor::GetPackageName() const
 	return Workspace->GetPackage()->GetName();
 }
 
+const FWorkspaceDocument& FWorkspaceEditor::GetFocusedWorkspaceDocument() const
+{
+	if (LastDocument.IsSet())
+	{
+		return LastDocument.GetValue();
+	}
+
+	static FWorkspaceDocument EmptyDocument;
+	return EmptyDocument;
+}
+
 void FWorkspaceEditor::BindCommands()
 {
 	const FWorkspaceAssetEditorCommands& Commands = FWorkspaceAssetEditorCommands::Get();
@@ -583,12 +605,12 @@ void FWorkspaceEditor::BindCommands()
 	);
 }
 
-void FWorkspaceEditor::ExtendMenu()
+void FWorkspaceEditor::ExtendMenu() const
 {
 	
 }
 
-void FWorkspaceEditor::ExtendToolbar()
+void FWorkspaceEditor::ExtendToolbar() const
 {
 	if (UToolMenu* Menu = UToolMenus::Get()->ExtendMenu(GetToolMenuToolbarName()))
 	{
@@ -633,7 +655,7 @@ void FWorkspaceEditor::UnregisterTabSpawners(const TSharedRef<FTabManager>& InTa
 
 FName FWorkspaceEditor::GetToolkitFName() const
 {
-	return FName("WorkspaceEditor");
+	return IWorkspaceEditor::GetWorkspaceEditorToolkitName();
 }
 
 FText FWorkspaceEditor::GetBaseToolkitName() const
@@ -692,7 +714,7 @@ void FWorkspaceEditor::CloseDocumentTab(const UObject* DocumentID) const
 	DocumentManager->CloseTab(Payload);
 }
 
-bool FWorkspaceEditor::InEditingMode() const
+bool FWorkspaceEditor::InEditingMode() const 
 {
 	return true;
 }
@@ -751,13 +773,15 @@ FText FWorkspaceEditor::GetToolkitToolTipText() const
 	UObject* const* WorkspaceObject = GetEditingObjects().FindByPredicate([](const UObject* Object) { return Object && Object->IsA<UWorkspace>(); });
 	check (WorkspaceObject != nullptr);
 	
-	FText FocussedAssetText = FText::FromName(NAME_None);	
-	if (const UObject* FocussedDocument = GetFocussedDocument())
+	FText FocusedAssetText = FText::FromName(NAME_None);
+
+	const FWorkspaceDocument& Document = GetFocusedWorkspaceDocument();
+	if (const UObject* FocusedObject = Document.GetObject())
 	{
-		 FocussedAssetText = GetLabelForObject(FocussedDocument);
+		 FocusedAssetText = GetLabelForObject(FocusedObject);
 	}
 	
-	return FText::Format(LOCTEXT("TookitTooltipFormat", "{0} ({1})"), GetToolTipTextForObject(*WorkspaceObject), FocussedAssetText);
+	return FText::Format(LOCTEXT("ToolkitTooltipFormat", "{0} ({1})"), GetToolTipTextForObject(*WorkspaceObject), FocusedAssetText);
 }
 
 void FWorkspaceEditor::RecordDocumentState(const TInstancedStruct<FWorkspaceDocumentState>& InState) const
@@ -766,13 +790,13 @@ void FWorkspaceEditor::RecordDocumentState(const TInstancedStruct<FWorkspaceDocu
 	State->DocumentStates.AddUnique(InState);
 }
 
-void FWorkspaceEditor::NavigateBack()
+void FWorkspaceEditor::NavigateBack() const
 {
 	const TSharedRef<FTabPayload_WorkspaceDocument> Payload = FTabPayload_WorkspaceDocument::Make(nullptr);
 	TSharedPtr<SDockTab> OpenedTab = DocumentManager->OpenDocument(Payload, FDocumentTracker::NavigateBackwards);
 }
 
-void FWorkspaceEditor::NavigateForward()
+void FWorkspaceEditor::NavigateForward() const
 {
 	const TSharedRef<FTabPayload_WorkspaceDocument> Payload = FTabPayload_WorkspaceDocument::Make(nullptr);
 	TSharedPtr<SDockTab> OpenedTab = DocumentManager->OpenDocument(Payload, FDocumentTracker::NavigateForwards);
@@ -844,7 +868,7 @@ bool FWorkspaceEditor::OnRequestClose(EAssetEditorCloseReason InCloseReason)
 
 		if(Result == FSuppressableWarningDialog::EResult::Confirm)
 		{
-			// Ensure we dont also 'save as' other externally linked assets at this point
+			// Ensure we do not also 'save as' other externally linked assets at this point
 			TGuardValue<bool> SaveWorkspaceOnly(bSavingTransientWorkspace, true);
 			SaveAssetAs_Execute();
 		}
@@ -888,9 +912,10 @@ bool FWorkspaceEditor::ShouldReopenEditorForSavedAsset(const UObject* Asset) con
 void FWorkspaceEditor::RemoveEditingObject(UObject* Object)
 {
 	IWorkspaceEditor::RemoveEditingObject(Object);
-	if (GetFocussedDocument() == Object)
+	if (LastDocument.IsSet() && LastDocument->GetObject() == Object)
 	{
-		SetFocussedDocument(nullptr);
+		FWorkspaceDocument EmptyDocument;
+		SetFocusedDocument(EmptyDocument);
 	}
 }
 }
diff --git a/Engine/Plugins/Experimental/Workspace/Source/WorkspaceEditor/Private/WorkspaceEditor.h b/Engine/Plugins/Experimental/Workspace/Source/WorkspaceEditor/Private/WorkspaceEditor.h
--- a/Engine/Plugins/Experimental/Workspace/Source/WorkspaceEditor/Private/WorkspaceEditor.h
+++ b/Engine/Plugins/Experimental/Workspace/Source/WorkspaceEditor/Private/WorkspaceEditor.h
@@ -5,6 +5,7 @@
 #include "CoreMinimal.h"
 #include "StructUtils/InstancedStruct.h"
 #include "IWorkspaceEditor.h"
+#include "IWorkspaceEditorModule.h"
 #include "WorkspaceAssetRegistryInfo.h"
 #include "WorkspaceEditorModeUILayer.h"
 #include "WorkflowOrientedApp/WorkflowTabManager.h"
@@ -71,7 +72,7 @@ private:
 	virtual void OnClose() override;
 	virtual void RegisterToolbar() override;
 	virtual bool ShouldReopenEditorForSavedAsset(const UObject* Asset) const override;
-	virtual void RemoveEditingObject(UObject* Object);
+	virtual void RemoveEditingObject(UObject* Object) override;
 
 	// FAssetEditorToolkit interface
 	virtual void GetSaveableObjects(TArray<UObject*>& OutObjects) const override;
@@ -98,7 +99,7 @@ private:
 	virtual void OpenExports(TConstArrayView<FWorkspaceOutlinerItemExport> InExports) override;
 	virtual void OpenObjects(TConstArrayView<UObject*> InObjects) override;
 	virtual void GetOpenedAssetsOfClass(TSubclassOf<UObject> InClass, TArray<UObject*>& OutObjects) const override;
-	virtual void GetAssets(TArray<FAssetData>& OutAssets) const override;
+	virtual void GetAssets(TArray<FAssetData>& OutAssets, bool bIncludeAssetReferences = false) const override;
 	virtual void CloseObjects(TConstArrayView<UObject*> InObjects) override;
 	virtual void SetDetailsObjects(const TArray<UObject*>& InObjects) override;
 	virtual void RefreshDetails() override;
@@ -106,19 +107,20 @@ private:
 	virtual bool GetOutlinerSelection(TArray<FWorkspaceOutlinerItemExport>& OutExports) const override;
 	virtual FOnOutlinerSelectionChanged& OnOutlinerSelectionChanged() override;
 	virtual void SetGlobalSelection(FGlobalSelectionId SelectionId, FOnClearGlobalSelection OnClearSelectionDelegate) override;
-	virtual const TObjectPtr<UObject> GetFocussedDocumentOfClass(const TObjectPtr<UClass> AssetClass) const override;
-	virtual FOnFocussedDocumentChanged& OnFocussedDocumentChanged() override;
+	virtual const FWorkspaceDocument& GetFocusedDocumentOfClass(const TObjectPtr<UClass> AssetClass) const override;
+	virtual FOnFocusedDocumentChanged& OnFocusedDocumentChanged() override;
 	virtual TSharedPtr<IDetailsView> GetDetailsView() override;
 	virtual UObject* GetWorkspaceAsset() const override;
 	virtual FString GetPackageName() const override;
+	virtual const FWorkspaceDocument& GetFocusedWorkspaceDocument() const override;
 
 	void HandleOutlinerSelectionChanged(TConstArrayView<FWorkspaceOutlinerItemExport> InExports);
 
 	void BindCommands();
 
-	void ExtendMenu();
+	void ExtendMenu() const;
 
-	void ExtendToolbar();
+	void ExtendToolbar() const;
 
 	void CloseDocumentTab(const UObject* DocumentID) const;
 
@@ -133,10 +135,10 @@ private:
 
 	void RecordDocumentState(const TInstancedStruct<FWorkspaceDocumentState>& InState) const;
 
-	void SetFocussedDocument(const TObjectPtr<UObject> InAsset);
+	void SetFocusedDocument(const UE::Workspace::FWorkspaceDocument& Document);
 
-	void NavigateBack();
-	void NavigateForward();
+	void NavigateBack() const;
+	void NavigateForward() const;
 
 	void SaveAssetEntries();
 	bool AreAssetEntriesModified() const;
@@ -159,7 +161,7 @@ private:
 	bool bSavingWorkspace = false;
 	bool bSavingAssetEntries = false;
 	bool bClosingDown = false;
-	bool bSettingFocussedDocument = false;
+	bool bSettingFocusedDocument = false;
 
 	TArray<TPair<FGlobalSelectionId, FOnClearGlobalSelection>> GlobalSelections;
 	FOnOutlinerSelectionChanged OnOutlinerSelectionChangedDelegate;
@@ -167,8 +169,10 @@ private:
 	int32 SelectionScopeDepth = 0;
 	
 	TArray<FWorkspaceOutlinerItemExport> LastSelectedExports;
-	FOnFocussedDocumentChanged OnFocussedDocumentChangedDelegate;
+	FOnFocusedDocumentChanged OnFocusedDocumentChangedDelegate;
 	TSharedPtr<FWorkspaceItem> EditorMenuCategory;
+
+	TOptional<FWorkspaceDocument> LastDocument;
 };
 
 }
\ No newline at end of file
diff --git a/Engine/Plugins/Experimental/Workspace/Source/WorkspaceEditor/Private/WorkspaceEditorModule.cpp b/Engine/Plugins/Experimental/Workspace/Source/WorkspaceEditor/Private/WorkspaceEditorModule.cpp
--- a/Engine/Plugins/Experimental/Workspace/Source/WorkspaceEditor/Private/WorkspaceEditorModule.cpp
+++ b/Engine/Plugins/Experimental/Workspace/Source/WorkspaceEditor/Private/WorkspaceEditorModule.cpp
@@ -19,6 +19,7 @@
 #include "Modules/ModuleManager.h"
 #include "Styling/AppStyle.h"
 #include "WorkspaceAssetReferenceItemDetails.h"
+#include "WorkspaceGroupItemDetails.h"
 
 #define LOCTEXT_NAMESPACE "WorkspaceEditorModule"
 
@@ -34,7 +35,7 @@ FWorkspaceDocument::FWorkspaceDocument(const FWorkspaceOutlinerItemExport& InExp
 	const FObjectDocumentArgs* FoundArgs = WorkspaceEditorModule.ObjectDocumentArgs.Find(InObject->GetClass()->GetClassPathName());
 	while(FoundArgs && FoundArgs->OnRedirectWorkspaceContext.IsBound())
 	{
-		InObject = FoundArgs->OnRedirectWorkspaceContext.Execute(InObject);
+		InObject = FoundArgs->OnRedirectWorkspaceContext.Execute(*this);
 		FoundArgs = WorkspaceEditorModule.ObjectDocumentArgs.Find(InObject->GetClass()->GetClassPathName());
 	}
 
@@ -54,7 +55,10 @@ void FWorkspaceEditorModule::StartupModule()
 	FWorkspaceAssetEditorCommands::Register();
 	
 	TSharedPtr<FWorkspaceAssetReferenceOutlinerItemDetails> ReferenceItemDetails = MakeShareable<FWorkspaceAssetReferenceOutlinerItemDetails>(new FWorkspaceAssetReferenceOutlinerItemDetails());
-	RegisterWorkspaceItemDetails(Workspace::FOutlinerItemDetailsId(FWorkspaceOutlinerAssetReferenceItemData::StaticStruct()->GetFName()), StaticCastSharedPtr<UE::Workspace::IWorkspaceOutlinerItemDetails>(ReferenceItemDetails));	
+	RegisterWorkspaceItemDetails(Workspace::FOutlinerItemDetailsId(FWorkspaceOutlinerAssetReferenceItemData::StaticStruct()->GetFName()), StaticCastSharedPtr<UE::Workspace::IWorkspaceOutlinerItemDetails>(ReferenceItemDetails));
+
+	TSharedPtr<FWorkspaceGroupOutlinerItemDetails> GroupItemDetails = MakeShareable<FWorkspaceGroupOutlinerItemDetails>(new FWorkspaceGroupOutlinerItemDetails());
+	RegisterWorkspaceItemDetails(Workspace::FOutlinerItemDetailsId(FWorkspaceOutlinerGroupItemData::StaticStruct()->GetFName()), StaticCastSharedPtr<UE::Workspace::IWorkspaceOutlinerItemDetails>(GroupItemDetails));
 }
 
 void FWorkspaceEditorModule::RegisterObjectDocumentType(const FTopLevelAssetPath& InClassPath, const FObjectDocumentArgs& InParams)
@@ -121,7 +125,7 @@ FObjectDocumentArgs FWorkspaceEditorModule::CreateGraphDocumentArgs(const FGraph
 	});
 	Args.OnGetTabName = FOnGetTabName::CreateLambda([](const FWorkspaceEditorContext& InContext)
 	{
-		if(UEdGraph* Graph = Cast<UEdGraph>(InContext.Document.Object))
+		if(UEdGraph* Graph = InContext.Document.GetTypedObject<UEdGraph>())
 		{
 			if (const UEdGraphSchema* Schema = Graph->GetSchema())
 			{
@@ -149,7 +153,7 @@ FObjectDocumentArgs FWorkspaceEditorModule::CreateGraphDocumentArgs(const FGraph
 			GraphDocument->GraphEditor->GetViewLocation(ViewLocation, ZoomAmount);
 		}
 	
-		return TInstancedStruct<FGraphDocumentState>::Make(InContext.Document.Object, InContext.Document.Export, FDeprecateSlateVector2D(ViewLocation), ZoomAmount);
+		return TInstancedStruct<FGraphDocumentState>::Make(InContext.Document.GetObject(), InContext.Document.Export, FDeprecateSlateVector2D(ViewLocation), ZoomAmount);
 	});
 	Args.OnSetDocumentState = FOnSetDocumentState::CreateLambda([](const FWorkspaceEditorContext& InContext, TSharedRef<SWidget> InWidget, const TInstancedStruct<FWorkspaceDocumentState>& InDocumentState)
 	{
@@ -242,7 +246,7 @@ bool FWorkspaceEditorModule::GetExportedAssetsForWorkspace(const FAssetData& InW
 		FName PackageName = *FPaths::GetBaseFilename(Package->GetName());
 		if (Factory->ConfigureProperties())
 		{
-			UWorkspace* NewWorkspace = CastChecked<UWorkspace>(Factory->FactoryCreateNew(UWorkspace::StaticClass(), Package, PackageName, RF_Public | RF_Standalone, NULL, GWarn));
+			UWorkspace* NewWorkspace = CastChecked<UWorkspace>(Factory->FactoryCreateNew(UWorkspace::StaticClass(), Package, PackageName, RF_Public | RF_Standalone, nullptr, GWarn));
 			NewWorkspace->AddAsset(InObject, false);
 			NewWorkspace->MarkPackageDirty();
 
@@ -255,7 +259,7 @@ bool FWorkspaceEditorModule::GetExportedAssetsForWorkspace(const FAssetData& InW
 		}
 	};
 
-	auto HandleExistingWorkspace = [InObject, &WorkspaceEditor](const FAssetData& InAssetData)
+	auto HandleExistingWorkspace = [&WorkspaceEditor](const FAssetData& InAssetData)
 	{
 		if(UWorkspace* ExistingWorkspace = Cast<UWorkspace>(InAssetData.GetAsset()))
 		{
@@ -294,13 +298,13 @@ bool FWorkspaceEditorModule::GetExportedAssetsForWorkspace(const FAssetData& InW
 	return WorkspaceEditor;
 }
 
-const FObjectDocumentArgs* FWorkspaceEditorModule::FindObjectDocumentType(const UObject* InObject) const
+const FObjectDocumentArgs* FWorkspaceEditorModule::FindObjectDocumentType(const FWorkspaceDocument& Document) const
 {
-	const FObjectDocumentArgs* FoundArgs = ObjectDocumentArgs.Find(InObject->GetClass()->GetClassPathName());
+	const FObjectDocumentArgs* FoundArgs = ObjectDocumentArgs.Find(Document.GetObject()->GetClass()->GetClassPathName());
 	while(FoundArgs && FoundArgs->OnRedirectWorkspaceContext.IsBound())
 	{
-		InObject = FoundArgs->OnRedirectWorkspaceContext.Execute(const_cast<UObject*>(InObject));
-		FoundArgs = InObject ? ObjectDocumentArgs.Find(InObject->GetClass()->GetClassPathName()) : nullptr;
+		UObject* Object = FoundArgs->OnRedirectWorkspaceContext.Execute(Document);
+		FoundArgs = Object ? ObjectDocumentArgs.Find(Object->GetClass()->GetClassPathName()) : nullptr;
 	}
 
 	return FoundArgs;
diff --git a/Engine/Plugins/Experimental/Workspace/Source/WorkspaceEditor/Private/WorkspaceEditorModule.h b/Engine/Plugins/Experimental/Workspace/Source/WorkspaceEditor/Private/WorkspaceEditorModule.h
--- a/Engine/Plugins/Experimental/Workspace/Source/WorkspaceEditor/Private/WorkspaceEditorModule.h
+++ b/Engine/Plugins/Experimental/Workspace/Source/WorkspaceEditor/Private/WorkspaceEditorModule.h
@@ -41,7 +41,7 @@ private:
 	virtual FOnExtendToolMenuContext& OnExtendToolMenuContext() override { return ExtendToolMenuContext; }
 
 	// Find an existing registered object document type. Note this redirects based on FObjectDocumentArgs::OnRedirectWorkspaceContext
-	const FObjectDocumentArgs* FindObjectDocumentType(const UObject* InObject) const;
+	const FObjectDocumentArgs* FindObjectDocumentType(const FWorkspaceDocument& Document) const;
 
 	// Find an existing registered document sub-object type.
 	const FDocumentSubObjectArgs* FindDocumentSubObjectType(const UObject* InObject) const;
@@ -78,7 +78,7 @@ private:
 
 	friend struct FAssetDocumentSummoner;
 	friend struct FWorkspaceOutlinerTreeItem;
-	friend class SWorkspaceOutlinerTreelabel;
+	friend class SWorkspaceOutlinerTreeLabel;
 	friend class FWorkspaceEditor;
 	friend class FWorkspaceOutlinerMode;
 	friend class FWorkspaceOutlinerSourceControlColumn;
diff --git a/Engine/Plugins/Experimental/Workspace/Source/WorkspaceEditor/Public/IWorkspaceEditor.h b/Engine/Plugins/Experimental/Workspace/Source/WorkspaceEditor/Public/IWorkspaceEditor.h
--- a/Engine/Plugins/Experimental/Workspace/Source/WorkspaceEditor/Public/IWorkspaceEditor.h
+++ b/Engine/Plugins/Experimental/Workspace/Source/WorkspaceEditor/Public/IWorkspaceEditor.h
@@ -17,7 +17,9 @@ namespace UE::Workspace
 typedef TWeakPtr<SWidget> FGlobalSelectionId;
 using FOnClearGlobalSelection = FSimpleDelegate;
 
-DECLARE_MULTICAST_DELEGATE_OneParam(FOnFocussedDocumentChanged, TObjectPtr<UObject>);
+struct FWorkspaceDocument;
+
+DECLARE_MULTICAST_DELEGATE_OneParam(FOnFocusedDocumentChanged, const FWorkspaceDocument&);
 
 // RAII helper allowing for a multi-widget selection scope within a WorkspaceEditor instance
 struct FWorkspaceEditorSelectionScope
@@ -53,7 +55,7 @@ public:
 	void GetOpenedAssets(TArray<UObject*>& OutObjects) const { GetOpenedAssetsOfClass(AssetClass::StaticClass(), OutObjects); }
 
 	// Get the current set of assets in this workspace editor 
-	virtual void GetAssets(TArray<FAssetData>& OutAssets) const = 0;
+	virtual void GetAssets(TArray<FAssetData>& OutAssets, bool bIncludeAssetReferences = false) const = 0;
 
 	// Close the supplied objects if they are open for editing within the workspace editor
 	virtual void CloseObjects(TConstArrayView<UObject*> InObjects) = 0;
@@ -70,24 +72,14 @@ public:
 	// Set the _current_ global selection (last SWidget with selection set) with delegate to clear it selection on next SetGlobalSelection()
 	virtual void SetGlobalSelection(FGlobalSelectionId SelectionId, FOnClearGlobalSelection OnClearSelectionDelegate) = 0;
 
-	// Get the currently focussed document. @return nullptr if the class does not match or no document is focussed
-	virtual const TObjectPtr<UObject> GetFocussedDocumentOfClass(const TObjectPtr<UClass> InClass) const = 0;
-
-	// Get the currently focussed document. @return nullptr if the class does not match or no document is focussed
-	template<typename AssetClass>
-	TObjectPtr<AssetClass> GetFocussedDocument() const
-	{
-		return Cast<AssetClass>(GetFocussedDocumentOfClass(AssetClass::StaticClass()));
-	}
+	// Get the currently focused document. @return unset if the class does not match or no document is focused
+	virtual const FWorkspaceDocument& GetFocusedDocumentOfClass(const TObjectPtr<UClass> InClass) const = 0;
 
-	// Get the currently focussed document. @return nullptr if no document is focussed
-	TObjectPtr<UObject> GetFocussedDocument() const
-	{
-		return GetFocussedDocumentOfClass(UObject::StaticClass());
-	}
+	// Returns the currently focused workspace document (last user activated tab), unset in case of none is focused
+	virtual const FWorkspaceDocument& GetFocusedWorkspaceDocument() const = 0;
 
-	// Multicast delegate broadcast whenever the document focussed inside of the WorkspaceEditor changes
-	virtual FOnFocussedDocumentChanged& OnFocussedDocumentChanged() = 0;
+	// Multicast delegate broadcast whenever the document focus of the WorkspaceEditor changes
+	virtual FOnFocusedDocumentChanged& OnFocusedDocumentChanged() = 0;
 
 	// Get the current single selection of the outliner.
 	// @return true if a single selection is active
@@ -105,6 +97,9 @@ public:
 
 	// Returns the name of the package where the workspace is located
 	virtual FString GetPackageName() const = 0;
+
+	// Name of EditorToolkit, used for testing against FAssetEditorToolkit::GetToolkitFName 
+	static FName GetWorkspaceEditorToolkitName() { return FName("WorkspaceEditor"); }
 };
 
 }
diff --git a/Engine/Plugins/Experimental/Workspace/Source/WorkspaceEditor/Public/IWorkspaceEditorModule.h b/Engine/Plugins/Experimental/Workspace/Source/WorkspaceEditor/Public/IWorkspaceEditorModule.h
--- a/Engine/Plugins/Experimental/Workspace/Source/WorkspaceEditor/Public/IWorkspaceEditorModule.h
+++ b/Engine/Plugins/Experimental/Workspace/Source/WorkspaceEditor/Public/IWorkspaceEditorModule.h
@@ -58,9 +58,13 @@ struct FWorkspaceDocument
 		return Object == InOther.Object && GetTypeHash(Export) == GetTypeHash(InOther.Export);
 	}
 
-	UObject* GetObject() const { check(Object); return Object; }
+	UObject* GetObject() const { return Object; }
+
+	template<class ObjectClass>
+	ObjectClass* GetTypedObject() const { return Cast<ObjectClass>(Object); }
+
 	template<class ObjectClass>
-	ObjectClass* GetTypedObject() const { return CastChecked<ObjectClass>(Object); }
+	ObjectClass* GetTypedObjectChecked() const { return CastChecked<ObjectClass>(Object); }
 };
 
 // Context passed to workspace editor delegates
@@ -90,7 +94,7 @@ struct FWorkspaceBreadcrumb : TSharedFromThis<FWorkspaceBreadcrumb>
 	FOnSaveBreadcrumb OnSave;
 };
 
-using FOnRedirectWorkspaceContext = TDelegate<UObject*(UObject*)>;
+using FOnRedirectWorkspaceContext = TDelegate<UObject*(const FWorkspaceDocument&)>;
 
 using FOnMakeDocumentWidget = TDelegate<TSharedRef<SWidget>(const FWorkspaceEditorContext&)>;
 
@@ -113,12 +117,12 @@ struct FObjectDocumentArgs
 {
 	FObjectDocumentArgs() = default;
 
-	FObjectDocumentArgs(FOnMakeDocumentWidget InOnMakeDocumentWidget, FName InSpawnLocation = WorkspaceTabs::TopMiddleDocumentArea)
+	FObjectDocumentArgs(const FOnMakeDocumentWidget& InOnMakeDocumentWidget, FName InSpawnLocation = WorkspaceTabs::TopMiddleDocumentArea)
 		: OnMakeDocumentWidget(InOnMakeDocumentWidget)
 		, SpawnLocation(InSpawnLocation)
 	{}
 
-	FObjectDocumentArgs(FOnRedirectWorkspaceContext InOnRedirectWorkspaceContext)
+	FObjectDocumentArgs(const FOnRedirectWorkspaceContext& InOnRedirectWorkspaceContext)
 		: OnRedirectWorkspaceContext(InOnRedirectWorkspaceContext)
 	{}
 	
@@ -146,7 +150,7 @@ struct FObjectDocumentArgs
 	// Delegate called to get the bread crumb trail for this document tab
 	FOnGetDocumentBreadcrumbTrail OnGetDocumentBreadcrumbTrail;
 
-	// EditorMode ID to be associated with this asset, will be used to try and activate matching editor mode when this asset is focussed
+	// EditorMode ID to be associated with this asset, will be used to try and activate matching editor mode when this asset is focused
 	FEditorModeID DocumentEditorMode = NAME_None;
 };
 
diff --git a/Engine/Plugins/Experimental/Workspace/Source/WorkspaceEditor/Public/IWorkspaceOutlinerItemDetails.h b/Engine/Plugins/Experimental/Workspace/Source/WorkspaceEditor/Public/IWorkspaceOutlinerItemDetails.h
--- a/Engine/Plugins/Experimental/Workspace/Source/WorkspaceEditor/Public/IWorkspaceOutlinerItemDetails.h
+++ b/Engine/Plugins/Experimental/Workspace/Source/WorkspaceEditor/Public/IWorkspaceOutlinerItemDetails.h
@@ -9,10 +9,10 @@ class UPackage;
 struct FToolMenuContext;
 struct FSlateBrush;
 
+#define LOCTEXT_NAMESPACE "IWorkspaceOutlinerItemDetails"
+
 namespace UE::Workspace
 {
-	constexpr bool bShowFullIdentifierInOutliner = false;	
-
 typedef FName FOutlinerItemDetailsId;
 static FOutlinerItemDetailsId MakeOutlinerDetailsId(const FWorkspaceOutlinerItemExport& InExport)
 {
@@ -25,8 +25,12 @@ public:
     virtual ~IWorkspaceOutlinerItemDetails() = default;
 	virtual FString GetDisplayString(const FWorkspaceOutlinerItemExport& Export) const
 	{			
-		return bShowFullIdentifierInOutliner ? FString::Printf(TEXT("%s - %s - %s"),
-		*Export.GetIdentifier().ToString(), *(Export.HasData() ? Export.GetData().GetScriptStruct()->GetFName() : NAME_None).ToString(), *Export.GetFirstAssetPath().ToString()) : Export.GetIdentifier().ToString();
+		return Export.GetIdentifier().ToString();
+	}
+
+	virtual FText GetToolTipText(const FWorkspaceOutlinerItemExport& Export) const
+	{			
+		return FText::Format(LOCTEXT("WorkspaceOutlinerToolTipFormat","Identifier: {0}\nPath: {1}\nType: {2}"), FText::FromString(GetDisplayString(Export)), FText::FromString(Export.GetFirstAssetPath().ToString()), Export.HasData() ? FText::FromString(Export.GetData().GetScriptStruct()->GetName()) : LOCTEXT("NoneTypeLabel", "None"));
 	}
 	
     virtual const FSlateBrush* GetItemIcon(const FWorkspaceOutlinerItemExport& Export) const { return nullptr; }
@@ -37,9 +41,22 @@ public:
     virtual bool CanRename(const FWorkspaceOutlinerItemExport& Export) const { return false; }
     virtual void Rename(const FWorkspaceOutlinerItemExport& Export, const FText& InName) const {}
     virtual bool ValidateName(const FWorkspaceOutlinerItemExport& Export, const FText& InName, FText& OutErrorMessage) const { return false; }
-    virtual UPackage* GetPackage(const FWorkspaceOutlinerItemExport& Export) const { return nullptr; }
+    virtual UPackage* GetPackage(const FWorkspaceOutlinerItemExport& Export) const
+	{
+		const FSoftObjectPath AssetPath = Export.GetFirstAssetPath();
+		if (AssetPath.IsValid())
+		{
+			if (UObject* AssetObject = AssetPath.ResolveObject())
+			{
+				return AssetObject->GetPackage();
+			}
+		}
+		return nullptr;
+	}
 	virtual bool HandleSelected(const FToolMenuContext& ToolMenuContext) const { return false; }
 	virtual bool IsExpandedByDefault() const { return true; }
 };
 
 }
+
+#undef LOCTEXT_NAMESPACE // "IWorkspaceOutlinerItemDetails"
\ No newline at end of file
diff --git a/Engine/Plugins/Experimental/Workspace/Source/WorkspaceEditor/Public/WorkspaceAssetRegistryInfo.h b/Engine/Plugins/Experimental/Workspace/Source/WorkspaceEditor/Public/WorkspaceAssetRegistryInfo.h
--- a/Engine/Plugins/Experimental/Workspace/Source/WorkspaceEditor/Public/WorkspaceAssetRegistryInfo.h
+++ b/Engine/Plugins/Experimental/Workspace/Source/WorkspaceEditor/Public/WorkspaceAssetRegistryInfo.h
@@ -3,6 +3,7 @@
 #pragma once
 
 #include "StructUtils/InstancedStruct.h"
+#include "Styling/SlateBrush.h"
 #include "WorkspaceAssetRegistryInfo.generated.h"
 
 USTRUCT()
@@ -101,7 +102,7 @@ protected:
 public:
 	FName GetIdentifier() const
 	{
-		// Path needs atleast two segments to contain a valid identifier
+		// Path needs at least two segments to contain a valid identifier
 		if(Path.PathSegments.Num() > 1)
 		{
 			return Path.PathSegments.Last();	
@@ -112,7 +113,7 @@ public:
 	
 	FName GetParentIdentifier() const
 	{
-		// Path needs atleast three segments to contain a valid _parent_ identifier
+		// Path needs at least three segments to contain a valid _parent_ identifier
 		const int32 NumSegments = Path.PathSegments.Num();
 		if (NumSegments > 2)
 		{
@@ -173,7 +174,7 @@ public:
 	// Returns the first FSoftObjectPath found in segments, starting from the end. e.g. "SoftObjectPath" - "Foo" - "SoftObjectPathTwo" - "Bar" will return SoftObjectPathTwo
 	FSoftObjectPath GetFirstAssetPath() const
 	{
-		// Path needs atleast one segment to contain a (potentially) valid asset path
+		// Path needs at least one segment to contain a (potentially) valid asset path
 		if(Path.PathSegments.Num() > 0)
 		{
 			for (int32 SegmentIndex = Path.PathSegments.Num() - 1; SegmentIndex >= 0; SegmentIndex--)
@@ -196,7 +197,7 @@ public:
 	//  - "Foo" - "SoftObjectPath" - "SoftObjectPathTwo" - "Bar" will return FSoftObjectPath()
 	FSoftObjectPath GetTopLevelAssetPath() const
 	{
-		// Path needs atleast one segment to contain a (potentially) valid asset path
+		// Path needs at least one segment to contain a (potentially) valid asset path
 		if(Path.PathSegments.Num() > 0)
 		{
 			return FSoftObjectPath(Path.PathSegments[0].ToString());
@@ -222,6 +223,31 @@ public:
 			}
 		}
 	}
+
+	// Returns all valid UObject-based assets found, and (optionally) loaded, in path segments, starting from the end. e.g:
+	//  - "SoftObjectPath" - "Foo" - "SoftObjectPathTwo" - "Bar" will return "SoftObjectPath" and "SoftObjectPathTwo" if their class matches AssetClass
+	template<typename AssetClass>
+	void GetAssetsOfType(TArray<AssetClass*>& OutAssets, bool bForceLoad = true) const
+	{
+		if(Path.PathSegments.Num() > 0)
+		{
+			for (int32 SegmentIndex = Path.PathSegments.Num() - 1; SegmentIndex >= 0; SegmentIndex--)
+			{
+				FSoftObjectPath ObjectPath = Path.PathSegments[SegmentIndex].ToString();
+				if (ObjectPath.IsValid() && ObjectPath.IsAsset())
+				{
+					AssetClass* Asset = Cast<AssetClass>(ObjectPath.ResolveObject());
+					Asset = Asset ? Asset : (bForceLoad ? Cast<AssetClass>(ObjectPath.TryLoad()) : nullptr);
+					
+					if (Asset)
+					{
+						OutAssets.Add(Asset);
+					}
+				}
+			}
+		}
+	}
+	
 	// Returns all valid FWorkspaceOutlinerItemExports found in path segments, starting from the end. e.g:
 	//  - "SoftObjectPath" - "Foo" - "SoftObjectPathTwo" - "Bar" will return "SoftObjectPath" - "Bar" , "SoftObjectPathTwo" - "Foo"
 	void GetExports(TArray<FWorkspaceOutlinerItemExport>& OutExports) const
@@ -258,7 +284,7 @@ public:
 	// Remove identifier segment to retrieve parent path hash
 	uint32 GetParentHash() const { return GetTypeHash(Path.RemoveSegment()); }
 
-	// Returns whether or not Data has any instanced struct setup
+	// Returns whether Data has any instanced struct setup
 	bool HasData() const { return Data.IsValid(); }
 	
 	const TInstancedStruct<FWorkspaceOutlinerItemData>& GetData() const { return Data; }
@@ -287,6 +313,10 @@ struct FWorkspaceOutlinerAssetReferenceItemData : public FWorkspaceOutlinerItemD
 	UPROPERTY()
 	FWorkspaceOutlinerItemExport ReferredExport;
 
+	// Whether this asset reference should be expanded in the Workspace outliner
+	UPROPERTY()
+	bool bShouldExpandReference = true;
+
 	UPROPERTY()
 	bool bRecursiveReference = false;
 
@@ -304,6 +334,20 @@ struct FWorkspaceOutlinerAssetReferenceItemData : public FWorkspaceOutlinerItemD
 	}
 };
 
+USTRUCT()
+struct FWorkspaceOutlinerGroupItemData : public FWorkspaceOutlinerItemData
+{
+	GENERATED_BODY()
+
+	FWorkspaceOutlinerGroupItemData() = default;
+
+	UPROPERTY()
+	FString GroupName;
+
+	UPROPERTY()
+	FSlateBrush GroupIcon;
+};
+
 namespace UE::Workspace
 {
 
diff --git a/Engine/Plugins/Experimental/Workspace/Source/WorkspaceEditor/Public/WorkspaceItemMenuContext.h b/Engine/Plugins/Experimental/Workspace/Source/WorkspaceEditor/Public/WorkspaceItemMenuContext.h
--- a/Engine/Plugins/Experimental/Workspace/Source/WorkspaceEditor/Public/WorkspaceItemMenuContext.h
+++ b/Engine/Plugins/Experimental/Workspace/Source/WorkspaceEditor/Public/WorkspaceItemMenuContext.h
@@ -3,6 +3,8 @@
 #pragma once
 
 #include "WorkspaceAssetRegistryInfo.h"
+#include "SceneOutlinerFwd.h"
+
 #include "WorkspaceItemMenuContext.generated.h"
 
 class FUICommandList;
@@ -13,5 +15,6 @@ class UWorkspaceItemMenuContext : public UObject
 	GENERATED_BODY()
 public:
 	TArray<FWorkspaceOutlinerItemExport> SelectedExports;
+	TArray<FSceneOutlinerTreeItemPtr> SelectedItems;
 	TWeakPtr<FUICommandList> WeakCommandList;
 };
diff --git a/Engine/Plugins/Experimental/Workspace/Source/WorkspaceEditor/Public/WorkspaceTabPayload.h b/Engine/Plugins/Experimental/Workspace/Source/WorkspaceEditor/Public/WorkspaceTabPayload.h
--- a/Engine/Plugins/Experimental/Workspace/Source/WorkspaceEditor/Public/WorkspaceTabPayload.h
+++ b/Engine/Plugins/Experimental/Workspace/Source/WorkspaceEditor/Public/WorkspaceTabPayload.h
@@ -20,7 +20,6 @@ struct FTabPayload_WorkspaceDocument : FTabPayload
 
 	static TSharedRef<FTabPayload_WorkspaceDocument> Make(const UObject* DocumentID)
 	{
-		check(DocumentID);
 		return MakeShareable(new FTabPayload_WorkspaceDocument(const_cast<UObject*>(DocumentID), FWorkspaceOutlinerItemExport()));
 	}
 
diff --git a/Engine/Plugins/Experimental/Workspace/Source/WorkspaceEditor/WorkspaceEditor.Build.cs b/Engine/Plugins/Experimental/Workspace/Source/WorkspaceEditor/WorkspaceEditor.Build.cs
--- a/Engine/Plugins/Experimental/Workspace/Source/WorkspaceEditor/WorkspaceEditor.Build.cs
+++ b/Engine/Plugins/Experimental/Workspace/Source/WorkspaceEditor/WorkspaceEditor.Build.cs
@@ -12,6 +12,12 @@ namespace UnrealBuildTool.Rules
 					"Core", 
 					"UnrealEd", 
 				});
+
+			PublicIncludePathModuleNames.AddRange(
+				new string[]
+				{
+					"SceneOutliner", 
+				});
 			
 			PrivateDependencyModuleNames.AddRange(
 				new string[]
