Changed default workspace editor layout

diff --git a/Engine/Plugins/Experimental/Workspace/Source/WorkspaceEditor/Private/WorkspaceEditor.cpp b/Engine/Plugins/Experimental/Workspace/Source/WorkspaceEditor/Private/WorkspaceEditor.cpp
--- a/Engine/Plugins/Experimental/Workspace/Source/WorkspaceEditor/Private/WorkspaceEditor.cpp
+++ b/Engine/Plugins/Experimental/Workspace/Source/WorkspaceEditor/Private/WorkspaceEditor.cpp
@@ -113,7 +113,7 @@ void FWorkspaceEditor::CreateWidgets()
 	check(DetailsView.IsValid());
 	WorkspaceEditorModule.ApplyWorkspaceDetailsCustomization(StaticCastWeakPtr<IWorkspaceEditor>(this->AsWeak()), DetailsView);
 
-	StandaloneDefaultLayout = FTabManager::NewLayout("Standalone_WorkspaceEditor_Layout_v1.2")
+	StandaloneDefaultLayout = FTabManager::NewLayout("Standalone_WorkspaceEditor_Layout_v2")
 	->AddArea
 	(
 		FTabManager::NewPrimaryArea()
@@ -132,7 +132,6 @@ void FWorkspaceEditor::CreateWidgets()
 					FTabManager::NewStack()
 					->SetSizeCoefficient(0.5f)
 					->SetHideTabWell(false)
-					->AddTab(WorkspaceTabs::WorkspaceView, ETabState::OpenedTab)
 					->AddTab(WorkspaceTabs::TopLeftDocumentArea, ETabState::ClosedTab)
 				)
 				->Split
@@ -173,8 +172,8 @@ void FWorkspaceEditor::CreateWidgets()
 					FTabManager::NewStack()
 					->SetSizeCoefficient(0.5f)
 					->SetHideTabWell(false)
+					->AddTab(WorkspaceTabs::WorkspaceView, ETabState::OpenedTab)
 					->AddTab(WorkspaceTabs::TopRightDocumentArea, ETabState::ClosedTab)
-					->AddTab(FBaseAssetToolkit::DetailsTabID, ETabState::OpenedTab)
 				)
 				->Split
 				(
@@ -182,6 +181,7 @@ void FWorkspaceEditor::CreateWidgets()
 					->SetSizeCoefficient(0.5f)
 					->SetHideTabWell(false)
 					->AddTab(WorkspaceTabs::BottomRightDocumentArea, ETabState::ClosedTab)
+					->AddTab(FBaseAssetToolkit::DetailsTabID, ETabState::OpenedTab)
 				)
 			)
 		)
