Fixed up TSubclassOf initialization to not use = nullptr;

diff --git a/Engine/Plugins/Experimental/Workspace/Source/WorkspaceEditor/Private/Workspace.h b/Engine/Plugins/Experimental/Workspace/Source/WorkspaceEditor/Private/Workspace.h
--- a/Engine/Plugins/Experimental/Workspace/Source/WorkspaceEditor/Private/Workspace.h
+++ b/Engine/Plugins/Experimental/Workspace/Source/WorkspaceEditor/Private/Workspace.h
@@ -150,7 +150,7 @@ class UWorkspace : public UObject
 
 	// Schema for this workspace
 	UPROPERTY(AssetRegistrySearchable)
-	TSubclassOf<UWorkspaceSchema> SchemaClass = nullptr;
+	TSubclassOf<UWorkspaceSchema> SchemaClass;
 
 	// State of the workspace, persisted to json
 	UPROPERTY(Transient)
diff --git a/Engine/Plugins/Experimental/Workspace/Source/WorkspaceEditor/Public/WorkspaceFactory.h b/Engine/Plugins/Experimental/Workspace/Source/WorkspaceEditor/Public/WorkspaceFactory.h
--- a/Engine/Plugins/Experimental/Workspace/Source/WorkspaceEditor/Public/WorkspaceFactory.h
+++ b/Engine/Plugins/Experimental/Workspace/Source/WorkspaceEditor/Public/WorkspaceFactory.h
@@ -37,7 +37,7 @@ private:
 	}
 
 	UPROPERTY()
-	TSubclassOf<UWorkspaceSchema> SchemaClass = nullptr;
+	TSubclassOf<UWorkspaceSchema> SchemaClass;
 
 	friend class UE::Workspace::FWorkspaceEditorModule;
 	friend class UE::Workspace::SWorkspacePicker;
