UAF Viewport 0.1

#virtualized

diff --git a/Engine/Plugins/Experimental/UAF/UAF/Source/UAF/Public/Component/AnimNextComponent.h b/Engine/Plugins/Experimental/UAF/UAF/Source/UAF/Public/Component/AnimNextComponent.h
--- a/Engine/Plugins/Experimental/UAF/UAF/Source/UAF/Public/Component/AnimNextComponent.h
+++ b/Engine/Plugins/Experimental/UAF/UAF/Source/UAF/Public/Component/AnimNextComponent.h
@@ -12,6 +12,11 @@
 
 #include "AnimNextComponent.generated.h"
 
+namespace UE::Workspace
+{
+	class IWorkspaceViewportController;
+}
+
 struct FAnimNextComponentInstanceData;
 class UAnimNextComponentWorldSubsystem;
 class UAnimNextModule;
@@ -27,6 +32,8 @@ namespace UE::UAF::UncookedOnly
 namespace UE::UAF::Editor
 {
 	class FAssetWizard;
+	class FSystemViewportController;
+	class FAnimGraphViewportController;
 }
 
 UCLASS(MinimalAPI, meta = (BlueprintSpawnableComponent))
@@ -210,6 +217,8 @@ private:
 	friend class UE::UAF::Editor::FAssetWizard;
 	// TODO ideally there would be a public API for the things FAnimMixerUtils does, so it wouldn't be necessary.
 	friend struct FAnimMixerUtils;
+	friend class UE::UAF::Editor::FSystemViewportController;
+	friend class UE::UAF::Editor::FAnimGraphViewportController;
 
 	// The UAF system that this component will run
 	UPROPERTY(EditAnywhere, Category="System", DisplayName="System")
diff --git a/Engine/Plugins/Experimental/UAF/UAF/Source/UAFEditor/Private/AnimNextEditorModule.cpp b/Engine/Plugins/Experimental/UAF/UAF/Source/UAFEditor/Private/AnimNextEditorModule.cpp
--- a/Engine/Plugins/Experimental/UAF/UAF/Source/UAFEditor/Private/AnimNextEditorModule.cpp
+++ b/Engine/Plugins/Experimental/UAF/UAF/Source/UAFEditor/Private/AnimNextEditorModule.cpp
@@ -75,6 +75,8 @@
 #include "UAFStyle.h"
 #include "Common/AnimNextAssetFindReplaceVariables.h"
 #include "Toolkits/AssetEditorToolkitMenuContext.h"
+#include "UAF/Viewport/SceneDescriptionTab.h"
+#include "UAF/Viewport/SystemViewportController.h"
 #include "Variables/Outliner/VariablesOutlinerCommands.h"
 #include "Wizard/AssetWizard.h"
 
@@ -182,6 +184,10 @@ void FAnimNextEditorModule::StartupModule()
 			TSharedRef<FAnimNextFindAndReplaceTabSummoner> FindAndReplaceTabSummoner = MakeShared<FAnimNextFindAndReplaceTabSummoner>(InEditorPtr);
 			TabFactories.RegisterFactory(FindAndReplaceTabSummoner);
 			FindAndReplaceTabSummoner->RegisterTabSpawner(InTabManager, nullptr);
+
+			TSharedRef<FUAFViewportSceneDescriptionTabSummoner> ViewportSceneDescriptionTabSummoner = MakeShared<FUAFViewportSceneDescriptionTabSummoner>(InEditorPtr);
+			TabFactories.RegisterFactory(ViewportSceneDescriptionTabSummoner);
+			ViewportSceneDescriptionTabSummoner->RegisterTabSpawner(InTabManager, nullptr);
 		});
 
 	WorkspaceEditorModule.OnExtendTabs().AddLambda([](FLayoutExtender& InLayoutExtender, TSharedPtr<UE::Workspace::IWorkspaceEditor> InEditorPtr)
@@ -499,6 +505,11 @@ void FAnimNextEditorModule::StartupModule()
 	{
 		return MakeShared<FAssetCompilationHandler>(InAsset);
 	}));
+
+	WorkspaceEditorModule.RegisterViewportControllerFactory(UAnimNextModule::StaticClass(), []()
+	{
+		return MakeUnique<FSystemViewportController>();
+	});
 }
 
 void FAnimNextEditorModule::ShutdownModule()
diff --git a/Engine/Plugins/Experimental/UAF/UAF/Source/UAFEditor/Private/UAF/Viewport/SViewportSceneDescription.cpp b/Engine/Plugins/Experimental/UAF/UAF/Source/UAFEditor/Private/UAF/Viewport/SViewportSceneDescription.cpp
new file mode 100644
--- /dev/null
+++ b/Engine/Plugins/Experimental/UAF/UAF/Source/UAFEditor/Private/UAF/Viewport/SViewportSceneDescription.cpp
@@ -0,0 +1,42 @@
+// Copyright Epic Games, Inc. All Rights Reserved.
+
+#include "SViewportSceneDescription.h"
+
+#include "IWorkspaceEditor.h"
+#include "UAF/Viewport/ViewportSceneDescription.h"
+#include "PropertyEditorModule.h"
+#include "IDetailsView.h"
+#include "Modules/ModuleManager.h"
+
+namespace UE::UAF::Editor
+{
+	void SViewportSceneDescription::Construct(const FArguments& InArgs, TSharedPtr<FAssetEditorToolkit> InAssetEditorToolkit)
+	{
+		FPropertyEditorModule& PropertyEditorModule = FModuleManager::LoadModuleChecked<FPropertyEditorModule>("PropertyEditor");
+
+		FDetailsViewArgs DetailsViewArgs;
+		DetailsViewArgs.bAllowSearch = false;
+		DetailsViewArgs.bAllowFavoriteSystem = false;
+		DetailsViewArgs.bShowOptions = false;
+		DetailsViewArgs.bShowObjectLabel = false;
+		DetailsViewArgs.bHideSelectionTip = true;
+		
+		TSharedRef<IDetailsView> DetailsView = PropertyEditorModule.CreateDetailView(DetailsViewArgs);
+		
+		ChildSlot
+		[
+			DetailsView
+		];
+
+		const TSharedPtr<UE::Workspace::IWorkspaceEditor> WorkspaceEditor = StaticCastSharedPtr<UE::Workspace::IWorkspaceEditor>(InAssetEditorToolkit);
+
+		UUAFViewportSceneDescription* SceneDescription = Cast<UUAFViewportSceneDescription>(WorkspaceEditor->GetSceneDescription());
+		if (!SceneDescription)
+		{
+			SceneDescription = NewObject<UUAFViewportSceneDescription>(GetTransientPackage());
+			WorkspaceEditor->SetSceneDescription(SceneDescription);
+		}
+		
+		DetailsView->SetObject(SceneDescription);
+	}
+}
diff --git a/Engine/Plugins/Experimental/UAF/UAF/Source/UAFEditor/Private/UAF/Viewport/SViewportSceneDescription.h b/Engine/Plugins/Experimental/UAF/UAF/Source/UAFEditor/Private/UAF/Viewport/SViewportSceneDescription.h
new file mode 100644
--- /dev/null
+++ b/Engine/Plugins/Experimental/UAF/UAF/Source/UAFEditor/Private/UAF/Viewport/SViewportSceneDescription.h
@@ -0,0 +1,20 @@
+// Copyright Epic Games, Inc. All Rights Reserved.
+
+#pragma once
+
+#include "Widgets/SCompoundWidget.h"
+
+class FAssetEditorToolkit;
+
+namespace UE::UAF::Editor
+{
+	class SViewportSceneDescription : public SCompoundWidget
+	{
+	public:
+		SLATE_BEGIN_ARGS(SViewportSceneDescription)
+		{}
+		SLATE_END_ARGS()
+		
+		void Construct(const FArguments& InArgs, TSharedPtr<FAssetEditorToolkit> InAssetEditorToolkit);
+	};
+}
\ No newline at end of file
diff --git a/Engine/Plugins/Experimental/UAF/UAF/Source/UAFEditor/Private/UAF/Viewport/SceneDescriptionTab.cpp b/Engine/Plugins/Experimental/UAF/UAF/Source/UAFEditor/Private/UAF/Viewport/SceneDescriptionTab.cpp
new file mode 100644
--- /dev/null
+++ b/Engine/Plugins/Experimental/UAF/UAF/Source/UAFEditor/Private/UAF/Viewport/SceneDescriptionTab.cpp
@@ -0,0 +1,37 @@
+// Copyright Epic Games, Inc. All Rights Reserved.
+
+#include "SceneDescriptionTab.h"
+
+#include "SViewportSceneDescription.h"
+#include "IWorkspaceEditor.h"
+
+#define LOCTEXT_NAMESPACE "FUAFViewportSceneDescriptionTabSummoner"
+
+const FLazyName ViewportSceneDescriptionTabName("ViewportSceneDescriptionTab");
+
+namespace UE::UAF::Editor
+{
+	FUAFViewportSceneDescriptionTabSummoner::FUAFViewportSceneDescriptionTabSummoner(TSharedPtr<UE::Workspace::IWorkspaceEditor> InHostingApp)
+	: FWorkflowTabFactory(ViewportSceneDescriptionTabName, StaticCastSharedPtr<FAssetEditorToolkit>(InHostingApp))
+	{	
+		TabLabel = LOCTEXT("TabLabel", "Viewport Scene");
+		TabIcon = FSlateIcon("EditorStyle", "LevelEditor.Tabs.Outliner");
+		ViewMenuDescription = LOCTEXT("TabMenuDescription", "Viewport Scene");
+		ViewMenuTooltip = LOCTEXT("TabToolTip", "Shows the Viewport Scene configuration tab.");
+		bIsSingleton = true;
+
+		Widget = SNew(SViewportSceneDescription, InHostingApp);
+	}
+
+	TSharedRef<SWidget> FUAFViewportSceneDescriptionTabSummoner::CreateTabBody(const FWorkflowTabSpawnInfo& Info) const
+	{
+		return Widget.ToSharedRef();
+	}
+
+	FText FUAFViewportSceneDescriptionTabSummoner::GetTabToolTipText(const FWorkflowTabSpawnInfo& Info) const
+	{
+		return LOCTEXT("TabToolTipText", "Viewport Scene");
+	}
+}
+
+#undef LOCTEXT_NAMESPACE
\ No newline at end of file
diff --git a/Engine/Plugins/Experimental/UAF/UAF/Source/UAFEditor/Private/UAF/Viewport/SceneDescriptionTab.h b/Engine/Plugins/Experimental/UAF/UAF/Source/UAFEditor/Private/UAF/Viewport/SceneDescriptionTab.h
new file mode 100644
--- /dev/null
+++ b/Engine/Plugins/Experimental/UAF/UAF/Source/UAFEditor/Private/UAF/Viewport/SceneDescriptionTab.h
@@ -0,0 +1,26 @@
+// Copyright Epic Games, Inc. All Rights Reserved.
+
+#pragma once
+
+#include "WorkflowOrientedApp/WorkflowTabFactory.h"
+
+namespace UE::Workspace
+{
+	class IWorkspaceEditor;
+}
+
+namespace UE::UAF::Editor
+{
+	struct FUAFViewportSceneDescriptionTabSummoner : public FWorkflowTabFactory
+	{
+		FUAFViewportSceneDescriptionTabSummoner(TSharedPtr<UE::Workspace::IWorkspaceEditor> InHostingApp);
+
+	private:
+		// FWorkflowTabFactory interface
+		virtual TSharedRef<SWidget> CreateTabBody(const FWorkflowTabSpawnInfo& Info) const override;
+		virtual FText GetTabToolTipText(const FWorkflowTabSpawnInfo& Info) const override;
+
+		// The widget this tab spawner wraps
+		TSharedPtr<class SCompoundWidget> Widget;
+	};
+}
\ No newline at end of file
diff --git a/Engine/Plugins/Experimental/UAF/UAF/Source/UAFEditor/Private/UAF/Viewport/SystemViewportController.cpp b/Engine/Plugins/Experimental/UAF/UAF/Source/UAFEditor/Private/UAF/Viewport/SystemViewportController.cpp
new file mode 100644
--- /dev/null
+++ b/Engine/Plugins/Experimental/UAF/UAF/Source/UAFEditor/Private/UAF/Viewport/SystemViewportController.cpp
@@ -0,0 +1,76 @@
+// Copyright Epic Games, Inc. All Rights Reserved.
+
+#include "SystemViewportController.h"
+
+#include "AdvancedPreviewScene.h"
+#include "AnimNextAssetWorkspaceAssetUserData.h"
+#include "IWorkspaceEditor.h"
+#include "Component/AnimNextComponent.h"
+#include "Engine/PreviewMeshCollection.h"
+#include "UAF/Viewport/ViewportSceneDescription.h"
+#include "Components/SkeletalMeshComponent.h"
+
+namespace UE::UAF::Editor
+{
+	void FSystemViewportController::OnEnter(const FViewportContext& InViewportContext)
+	{
+		UAnimNextModule* System = Cast<UAnimNextModule>(InViewportContext.OutlinerObject);
+		UUAFViewportSceneDescription* UAFViewportSceneDescription = CastChecked<UUAFViewportSceneDescription>(InViewportContext.SceneDescription);
+		
+		const TObjectPtr<UWorld> PreviewWorld = InViewportContext.PreviewScene->GetWorld();
+
+		const auto AddActor = [&](USkeletalMesh* InSkeletalMesh)
+			{
+				AActor* PreviewActor = PreviewWorld->SpawnActor<AActor>(FVector::ZeroVector, FRotator::ZeroRotator);
+				check(PreviewActor);
+
+				USkeletalMeshComponent* SkeletalMeshComponent = NewObject<USkeletalMeshComponent>(PreviewActor);
+				InViewportContext.PreviewScene->AddComponent(SkeletalMeshComponent, FTransform::Identity);
+				PreviewActor->SetRootComponent(SkeletalMeshComponent);
+				SkeletalMeshComponent->SetEnableAnimation(false);
+
+				SkeletalMeshComponent->SetSkeletalMesh(InSkeletalMesh);
+
+				UAnimNextComponent* UAFComponent = NewObject<UAnimNextComponent>(PreviewActor);
+				PreviewActor->AddInstanceComponent(UAFComponent);
+				UAFComponent->SetModule(System);
+				UAFComponent->RegisterComponent();
+				UAFComponent->InitializeComponent();
+			
+				PreviewActors.Add(PreviewActor);
+			};
+
+		if (UAFViewportSceneDescription->SkeletalMesh)
+		{
+			AddActor(UAFViewportSceneDescription->SkeletalMesh);
+		}
+
+		if (UAFViewportSceneDescription->AdditionalMeshes)
+		{
+			for (int32 MeshIndex = 0; MeshIndex < UAFViewportSceneDescription->AdditionalMeshes->SkeletalMeshes.Num(); ++MeshIndex)
+			{
+				const TSoftObjectPtr<USkeletalMesh> SkelMeshSoftPtr = UAFViewportSceneDescription->AdditionalMeshes->SkeletalMeshes[MeshIndex].SkeletalMesh;
+				
+				if (!SkelMeshSoftPtr.IsNull())
+				{
+					USkeletalMesh* SkeletalMesh = SkelMeshSoftPtr.LoadSynchronous();
+					if (SkeletalMesh)
+					{
+						AddActor(SkeletalMesh);
+					}
+				}
+			}
+		}
+	}
+	
+	void FSystemViewportController::OnExit()
+	{
+		for (AActor* PreviewActor : PreviewActors)
+		{
+			PreviewActor->Destroy();
+			PreviewActor = nullptr;
+		}
+		
+		PreviewActors.Empty();
+	}
+}
diff --git a/Engine/Plugins/Experimental/UAF/UAF/Source/UAFEditor/Private/UAF/Viewport/SystemViewportController.h b/Engine/Plugins/Experimental/UAF/UAF/Source/UAFEditor/Private/UAF/Viewport/SystemViewportController.h
new file mode 100644
--- /dev/null
+++ b/Engine/Plugins/Experimental/UAF/UAF/Source/UAFEditor/Private/UAF/Viewport/SystemViewportController.h
@@ -0,0 +1,16 @@
+// Copyright Epic Games, Inc. All Rights Reserved.
+
+#pragma once
+
+#include "WorkspaceViewportController.h"
+
+namespace UE::UAF::Editor
+{
+	class FSystemViewportController : public Workspace::IWorkspaceViewportController
+	{
+		virtual void OnEnter(const FViewportContext& InViewportContext) override;
+		virtual void OnExit() override;
+
+		TArray<AActor*> PreviewActors;
+	};
+}
\ No newline at end of file
diff --git a/Engine/Plugins/Experimental/UAF/UAF/Source/UAFEditor/Private/UAF/Viewport/ViewportSceneDescription.cpp b/Engine/Plugins/Experimental/UAF/UAF/Source/UAFEditor/Private/UAF/Viewport/ViewportSceneDescription.cpp
new file mode 100644
--- /dev/null
+++ b/Engine/Plugins/Experimental/UAF/UAF/Source/UAFEditor/Private/UAF/Viewport/ViewportSceneDescription.cpp
@@ -0,0 +1,20 @@
+// Copyright Epic Games, Inc. All Rights Reserved.
+
+#include "UAF/Viewport/ViewportSceneDescription.h"
+
+#include "PropertyCustomizationHelpers.h"
+#include "IWorkspaceEditor.h"
+#include "Animation/DebugSkelMeshComponent.h"
+#include "Component/AnimNextComponent.h"
+
+void UUAFViewportSceneDescription::PostEditChangeProperty(FPropertyChangedEvent& PropertyChangedEvent)
+{
+	Super::PostEditChangeProperty(PropertyChangedEvent);
+
+	const FName PropertyName = (PropertyChangedEvent.MemberProperty != nullptr) ? PropertyChangedEvent.MemberProperty->GetFName() : NAME_None;
+	if (PropertyName == GET_MEMBER_NAME_CHECKED(UUAFViewportSceneDescription, SkeletalMesh) || PropertyName == GET_MEMBER_NAME_CHECKED(UUAFViewportSceneDescription, AdditionalMeshes))
+	{
+		BroadcastOnConfigChanged();
+	}
+}
+
diff --git a/Engine/Plugins/Experimental/UAF/UAF/Source/UAFEditor/Private/Workspace/AnimNextWorkspaceSchema.h b/Engine/Plugins/Experimental/UAF/UAF/Source/UAFEditor/Private/Workspace/AnimNextWorkspaceSchema.h
--- a/Engine/Plugins/Experimental/UAF/UAF/Source/UAFEditor/Private/Workspace/AnimNextWorkspaceSchema.h
+++ b/Engine/Plugins/Experimental/UAF/UAF/Source/UAFEditor/Private/Workspace/AnimNextWorkspaceSchema.h
@@ -1,4 +1,4 @@
-﻿// Copyright Epic Games, Inc. All Rights Reserved.
+// Copyright Epic Games, Inc. All Rights Reserved.
 
 #pragma once
 
@@ -16,4 +16,5 @@ class UAnimNextWorkspaceSchema : public UWorkspaceSchema
 	virtual TConstArrayView<FTopLevelAssetPath> GetSupportedAssetClassPaths() const override;
 	virtual void OnSaveWorkspaceState(TSharedRef<UE::Workspace::IWorkspaceEditor> InWorkspaceEditor, FInstancedStruct& OutWorkspaceState) const override;
 	virtual void OnLoadWorkspaceState(TSharedRef<UE::Workspace::IWorkspaceEditor> InWorkspaceEditor, const FInstancedStruct& InWorkspaceState) const override;
+	virtual bool SupportsViewport() const { return true; };
 };
\ No newline at end of file
diff --git a/Engine/Plugins/Experimental/UAF/UAF/Source/UAFEditor/Public/UAF/Viewport/ViewportSceneDescription.h b/Engine/Plugins/Experimental/UAF/UAF/Source/UAFEditor/Public/UAF/Viewport/ViewportSceneDescription.h
new file mode 100644
--- /dev/null
+++ b/Engine/Plugins/Experimental/UAF/UAF/Source/UAFEditor/Public/UAF/Viewport/ViewportSceneDescription.h
@@ -0,0 +1,32 @@
+// Copyright Epic Games, Inc. All Rights Reserved.
+
+#pragma once
+
+#include "CoreMinimal.h"
+#include "WorkspaceViewportSceneDescription.h"
+
+#include "ViewportSceneDescription.generated.h"
+
+class FAdvancedPreviewScene;
+class USkeletalMesh;
+class UPreviewMeshCollection;
+
+namespace UE::Workspace
+{
+	class IWorkspaceEditor;
+}
+
+UCLASS()
+class UUAFViewportSceneDescription : public UWorkspaceViewportSceneDescription
+{
+	GENERATED_BODY()
+	
+public:
+	virtual void PostEditChangeProperty(FPropertyChangedEvent& PropertyChangedEvent) override;
+	
+	UPROPERTY(EditAnywhere, Category = Default)
+	TObjectPtr<USkeletalMesh> SkeletalMesh;
+
+	UPROPERTY(EditAnywhere, Category = Default)
+	TObjectPtr<UPreviewMeshCollection> AdditionalMeshes;
+};
\ No newline at end of file
diff --git a/Engine/Plugins/Experimental/UAF/UAF/Source/UAFEditor/UAFEditor.Build.cs b/Engine/Plugins/Experimental/UAF/UAF/Source/UAFEditor/UAFEditor.Build.cs
--- a/Engine/Plugins/Experimental/UAF/UAF/Source/UAFEditor/UAFEditor.Build.cs
+++ b/Engine/Plugins/Experimental/UAF/UAF/Source/UAFEditor/UAFEditor.Build.cs
@@ -57,7 +57,8 @@ namespace UnrealBuildTool.Rules
 					"GameplayInsights",
 					"Projects",
 					"Persona",
-					"DeveloperSettings"
+					"DeveloperSettings",
+					"AdvancedPreviewScene"
 				}
 			);
 		}
diff --git a/Engine/Plugins/Experimental/UAF/UAFAnimGraph/Source/UAFAnimGraphEditor/Private/AnimGraphViewportController.cpp b/Engine/Plugins/Experimental/UAF/UAFAnimGraph/Source/UAFAnimGraphEditor/Private/AnimGraphViewportController.cpp
new file mode 100644
--- /dev/null
+++ b/Engine/Plugins/Experimental/UAF/UAFAnimGraph/Source/UAFAnimGraphEditor/Private/AnimGraphViewportController.cpp
@@ -0,0 +1,80 @@
+// Copyright Epic Games, Inc. All Rights Reserved.
+
+#include "AnimGraphViewportController.h"
+
+#include "AdvancedPreviewScene.h"
+#include "AnimNextAnimGraphWorkspaceAssetUserData.h"
+#include "Component/AnimNextComponent.h"
+#include "UAF/Viewport/ViewportSceneDescription.h"
+#include "Engine/PreviewMeshCollection.h"
+#include "Components/SkeletalMeshComponent.h"
+
+namespace UE::UAF::Editor
+{
+	void FAnimGraphViewportController::OnEnter(const FViewportContext& InViewportContext)
+	{
+		UAnimNextAnimationGraph* Graph = CastChecked<UAnimNextAnimationGraph>(InViewportContext.OutlinerObject);
+		UUAFViewportSceneDescription* UAFViewportSceneDescription = CastChecked<UUAFViewportSceneDescription>(InViewportContext.SceneDescription);
+
+		// TODO: Consider building this module programatically
+		TObjectPtr<UAnimNextModule> Module = LoadObject<UAnimNextModule>(GetTransientPackage(), TEXT("/UAFAnimGraph/Internal/S_SingleGraph.S_SingleGraph"));
+		
+		const TObjectPtr<UWorld> PreviewWorld = InViewportContext.PreviewScene->GetWorld();
+
+		const auto AddActor = [&](USkeletalMesh* InSkeletalMesh)
+			{
+				AActor* PreviewActor = PreviewWorld->SpawnActor<AActor>(AActor::StaticClass());
+				check(PreviewActor);
+
+				USkeletalMeshComponent* SkeletalMeshComponent = NewObject<USkeletalMeshComponent>(PreviewActor);
+				InViewportContext.PreviewScene->AddComponent(SkeletalMeshComponent, FTransform::Identity);
+				PreviewActor->SetRootComponent(SkeletalMeshComponent);
+				SkeletalMeshComponent->SetEnableAnimation(false);
+				SkeletalMeshComponent->SetSkeletalMesh(InSkeletalMesh);
+
+				UAnimNextComponent* UAFComponent = NewObject<UAnimNextComponent>(PreviewActor);
+				PreviewActor->AddInstanceComponent(UAFComponent);
+				UAFComponent->SetModule(Module);
+				UAFComponent->RegisterComponent();
+				UAFComponent->InitializeComponent();
+		
+				FAnimNextVariableReference VariableReference("Graph", Module);
+				check(UAFComponent->SetVariable<TObjectPtr<UObject>>(VariableReference, Graph));
+
+				PreviewActors.Add(PreviewActor);
+			};
+
+		if (UAFViewportSceneDescription->SkeletalMesh)
+		{
+			AddActor(UAFViewportSceneDescription->SkeletalMesh);
+		}
+
+		if (UAFViewportSceneDescription->AdditionalMeshes)
+		{
+			for (int32 MeshIndex = 0; MeshIndex < UAFViewportSceneDescription->AdditionalMeshes->SkeletalMeshes.Num(); ++MeshIndex)
+			{
+				const TSoftObjectPtr<USkeletalMesh> SkelMeshSoftPtr = UAFViewportSceneDescription->AdditionalMeshes->SkeletalMeshes[MeshIndex].SkeletalMesh;
+				
+				if (!SkelMeshSoftPtr.IsNull())
+				{
+					USkeletalMesh* SkeletalMesh = SkelMeshSoftPtr.LoadSynchronous();
+					if (SkeletalMesh)
+					{
+						AddActor(SkeletalMesh);
+					}
+				}
+			}
+		}
+	}
+	
+	void FAnimGraphViewportController::OnExit()
+	{
+		for (AActor* PreviewActor : PreviewActors)
+		{
+			PreviewActor->Destroy();
+			PreviewActor = nullptr;
+		}
+		
+		PreviewActors.Empty();
+	}
+}
diff --git a/Engine/Plugins/Experimental/UAF/UAFAnimGraph/Source/UAFAnimGraphEditor/Private/AnimGraphViewportController.h b/Engine/Plugins/Experimental/UAF/UAFAnimGraph/Source/UAFAnimGraphEditor/Private/AnimGraphViewportController.h
new file mode 100644
--- /dev/null
+++ b/Engine/Plugins/Experimental/UAF/UAFAnimGraph/Source/UAFAnimGraphEditor/Private/AnimGraphViewportController.h
@@ -0,0 +1,16 @@
+// Copyright Epic Games, Inc. All Rights Reserved.
+
+#pragma once
+
+#include "WorkspaceViewportController.h"
+
+namespace UE::UAF::Editor
+{
+	class FAnimGraphViewportController : public Workspace::IWorkspaceViewportController
+	{
+		virtual void OnEnter(const FViewportContext& InViewportContext) override;
+		virtual void OnExit() override;
+
+		TArray<AActor*> PreviewActors;
+	};
+}
\ No newline at end of file
diff --git a/Engine/Plugins/Experimental/UAF/UAFAnimGraph/Source/UAFAnimGraphEditor/Private/AnimNextAnimGraphEditorModule.cpp b/Engine/Plugins/Experimental/UAF/UAFAnimGraph/Source/UAFAnimGraphEditor/Private/AnimNextAnimGraphEditorModule.cpp
--- a/Engine/Plugins/Experimental/UAF/UAFAnimGraph/Source/UAFAnimGraphEditor/Private/AnimNextAnimGraphEditorModule.cpp
+++ b/Engine/Plugins/Experimental/UAF/UAFAnimGraph/Source/UAFAnimGraphEditor/Private/AnimNextAnimGraphEditorModule.cpp
@@ -1,5 +1,6 @@
 // Copyright Epic Games, Inc. All Rights Reserved.
 
+#include "AdvancedPreviewScene.h"
 #include "AnimNextAnimGraphWorkspaceAssetUserData.h"
 #include "Modules/ModuleInterface.h"
 #include "Modules/ModuleManager.h"
@@ -11,6 +12,7 @@
 #include "PropertyEditorModule.h"
 #include "TraitStackEditor.h"
 #include "AnimGraphUncookedOnlyUtils.h"
+#include "AnimGraphViewportController.h"
 #include "Entries/AnimNextAnimationGraphEntry.h"
 #include "Features/IModularFeatures.h"
 #include "Framework/Docking/LayoutExtender.h"
@@ -40,6 +42,8 @@
 #include "RewindDebugger/EvaluationProgramTrack.h"
 #include "RewindDebugger/SequenceInfoTrack.h"
 #include "PersonaModule.h"
+#include "Component/AnimNextComponent.h"
+#include "Engine/PreviewMeshCollection.h"
 #include "Graph/AnimNextGraphPanelPinFactory.h"
 #include "Templates/UAFGraphNodeTemplate.h"
 #include "Factory/AnimNextFactoryParams.h"
@@ -117,7 +121,12 @@ class FAnimNextAnimGraphEditorModule : public IModuleInterface
 			}));
 
 		WorkspaceEditorModule.RegisterObjectDocumentType(FTopLevelAssetPath(TEXT("/Script/UAFAnimGraph.AnimNextAnimationGraph")), AnimNextAnimationGraphDocumentArgs);
-		
+
+		WorkspaceEditorModule.RegisterViewportControllerFactory(UAnimNextAnimationGraph::StaticClass(), []()
+			{
+				return MakeUnique<FAnimGraphViewportController>();
+			});
+
 /*		WorkspaceEditorModule.OnRegisterTabsForEditor().AddLambda([](FWorkflowAllowedTabSet& TabFactories, const TSharedRef<FTabManager>& InTabManager, TSharedPtr<UE::Workspace::IWorkspaceEditor> InEditorPtr)
 		{
 			TSharedRef<FTraitEditorTabSummoner> TraitEditorTabSummoner = MakeShared<FTraitEditorTabSummoner>(InEditorPtr);
@@ -171,6 +180,7 @@ class FAnimNextAnimGraphEditorModule : public IModuleInterface
 			if(UObjectInitialized())
 			{
 				WorkspaceEditorModule.UnregisterWorkspaceItemDetails(Workspace::FOutlinerItemDetailsId(FAnimNextAnimationGraphOutlinerData::StaticStruct()->GetFName()));
+				WorkspaceEditorModule.UnregisterViewportControllerFactory(UAnimNextAnimationGraph::StaticClass());
 			}
 		}
 
diff --git a/Engine/Plugins/Experimental/Workspace/Source/WorkspaceEditor/Private/WorkspaceAssetViewportClient.cpp b/Engine/Plugins/Experimental/Workspace/Source/WorkspaceEditor/Private/WorkspaceAssetViewportClient.cpp
new file mode 100644
--- /dev/null
+++ b/Engine/Plugins/Experimental/Workspace/Source/WorkspaceEditor/Private/WorkspaceAssetViewportClient.cpp
@@ -0,0 +1,18 @@
+// Copyright Epic Games, Inc. All Rights Reserved.
+
+#include "WorkspaceAssetViewportClient.h"
+
+#include "PreviewScene.h"
+
+namespace UE::Workspace
+{
+	void FWorkspaceAssetViewportClient::Tick(float DeltaTime)
+	{
+		if (PreviewScene)
+		{
+			PreviewScene->GetWorld()->Tick(LEVELTICK_All, DeltaTime);
+		}
+
+		FEditorViewportClient::Tick(DeltaTime);
+	}
+}
\ No newline at end of file
diff --git a/Engine/Plugins/Experimental/Workspace/Source/WorkspaceEditor/Private/WorkspaceAssetViewportClient.h b/Engine/Plugins/Experimental/Workspace/Source/WorkspaceEditor/Private/WorkspaceAssetViewportClient.h
new file mode 100644
--- /dev/null
+++ b/Engine/Plugins/Experimental/Workspace/Source/WorkspaceEditor/Private/WorkspaceAssetViewportClient.h
@@ -0,0 +1,20 @@
+// Copyright Epic Games, Inc. All Rights Reserved.
+
+#pragma once
+
+#include "EditorViewportClient.h"
+
+namespace UE::Workspace
+{
+	class FWorkspaceAssetViewportClient : public FEditorViewportClient
+	{
+	public:
+		FWorkspaceAssetViewportClient(FEditorModeTools* InModeTools, FPreviewScene* InPreviewScene)
+			: FEditorViewportClient(InModeTools, InPreviewScene)
+		{
+		}
+		
+		// FEditorViewportClient interface
+		virtual void Tick(float DeltaTime) override;
+	};
+}
\ No newline at end of file
diff --git a/Engine/Plugins/Experimental/Workspace/Source/WorkspaceEditor/Private/WorkspaceEditor.cpp b/Engine/Plugins/Experimental/Workspace/Source/WorkspaceEditor/Private/WorkspaceEditor.cpp
--- a/Engine/Plugins/Experimental/Workspace/Source/WorkspaceEditor/Private/WorkspaceEditor.cpp
+++ b/Engine/Plugins/Experimental/Workspace/Source/WorkspaceEditor/Private/WorkspaceEditor.cpp
@@ -20,10 +20,15 @@
 #include "AssetEditorModeManager.h"
 #include "Selection.h"
 #include "EditorModeManager.h"
+#include "SAssetEditorViewport.h"
 #include "Tools/EdModeInteractiveToolsContext.h"
 #include "SGraphDocument.h"
 #include "WorkspaceSchema.h"
 #include "WorkspaceTabPayload.h"
+#include "AdvancedPreviewScene.h"
+#include "WorkspaceAssetViewportClient.h"
+#include "WorkspaceViewportController.h"
+#include "Viewports.h"
 
 #define LOCTEXT_NAMESPACE "WorkspaceEditor"
 
@@ -113,7 +118,7 @@ void FWorkspaceEditor::CreateWidgets()
 	check(DetailsView.IsValid());
 	WorkspaceEditorModule.ApplyWorkspaceDetailsCustomization(StaticCastWeakPtr<IWorkspaceEditor>(this->AsWeak()), DetailsView);
 
-	StandaloneDefaultLayout = FTabManager::NewLayout("Standalone_WorkspaceEditor_Layout_v2")
+	StandaloneDefaultLayout = FTabManager::NewLayout("Standalone_WorkspaceEditor_Layout_v3")
 	->AddArea
 	(
 		FTabManager::NewPrimaryArea()
@@ -133,6 +138,7 @@ void FWorkspaceEditor::CreateWidgets()
 					->SetSizeCoefficient(0.5f)
 					->SetHideTabWell(false)
 					->AddTab(WorkspaceTabs::TopLeftDocumentArea, ETabState::ClosedTab)
+					->AddTab(FBaseAssetToolkit::ViewportTabID, ETabState::OpenedTab)
 				)
 				->Split
 				(
@@ -242,6 +248,12 @@ void FWorkspaceEditor::OnToolkitHostingStarted(const TSharedRef<IToolkit>& Toolk
 	RegenerateMenusAndToolbars();
 }
 
+void FWorkspaceEditor::AddReferencedObjects(FReferenceCollector& Collector)
+{
+	Collector.AddReferencedObject(Workspace);
+	Collector.AddReferencedObject(SceneDescription);
+}
+
 void FWorkspaceEditor::RestoreEditedObjectState()
 {
 	UWorkspaceState* State = Workspace->GetState();
@@ -528,6 +540,7 @@ void FWorkspaceEditor::SetFocusedDocument(const UE::Workspace::FWorkspaceDocumen
 		LastDocument.Reset();
 	}
 
+	ReinitViewportScene();
 	OnFocusedDocumentChangedDelegate.Broadcast(Document);
 }
 
@@ -580,7 +593,37 @@ const FWorkspaceDocument& FWorkspaceEditor::GetFocusedWorkspaceDocument() const
 	static FWorkspaceDocument EmptyDocument;
 	return EmptyDocument;
 }
+	
+void FWorkspaceEditor::ReinitViewportScene()
+{
+	if (ActiveViewportController)
+	{
+		ActiveViewportController->OnExit();
+		ActiveViewportController = nullptr;
+	}
+
+	const FSoftObjectPath AssetPath = GetFocusedWorkspaceDocument().Export.GetFirstAssetPath();
+
+	if (AssetPath.IsValid())
+	{
+		if (UObject* AssetObject = AssetPath.TryLoad())
+		{
+			FWorkspaceEditorModule& WorkspaceEditorModule = FModuleManager::Get().LoadModuleChecked<FWorkspaceEditorModule>("WorkspaceEditor");
+			ActiveViewportController = WorkspaceEditorModule.CreateViewportController(AssetObject->GetClass());
 
+			if (ActiveViewportController)
+			{
+				IWorkspaceViewportController::FViewportContext ViewportContext;
+				ViewportContext.PreviewScene = static_cast<FAdvancedPreviewScene*>(ViewportClient->GetPreviewScene());
+				ViewportContext.OutlinerObject = AssetObject;
+				ViewportContext.SceneDescription = SceneDescription;
+				
+				ActiveViewportController->OnEnter(ViewportContext);
+			}
+		}
+	}
+}
+	
 void FWorkspaceEditor::BindCommands()
 {
 	const FWorkspaceAssetEditorCommands& Commands = FWorkspaceAssetEditorCommands::Get();
@@ -646,6 +689,11 @@ void FWorkspaceEditor::RegisterTabSpawners(const TSharedRef<FTabManager>& InTabM
 	{
 		WorkspaceEditorModule.OnRegisterTabsForEditor().Broadcast(TabFactories, InTabManager, StaticCastSharedPtr<IWorkspaceEditor>(this->AsShared().ToSharedPtr()));
 	}
+
+	if (!GetSchema()->SupportsViewport())
+	{
+		InTabManager->UnregisterTabSpawner(ViewportTabID);
+	}
 }
 
 void FWorkspaceEditor::UnregisterTabSpawners(const TSharedRef<FTabManager>& InTabManager)
@@ -918,6 +966,48 @@ void FWorkspaceEditor::RemoveEditingObject(UObject* Object)
 		SetFocusedDocument(EmptyDocument);
 	}
 }
+
+AssetEditorViewportFactoryFunction FWorkspaceEditor::GetViewportDelegate()
+{
+	return [this](FAssetEditorViewportConstructionArgs InArgs)
+		{
+			TSharedRef<SAssetEditorViewport> Viewport = SNew(SAssetEditorViewport, InArgs)
+				.EditorViewportClient(ViewportClient);
+
+			return Viewport;
+		};
+}
+	
+TSharedPtr<FEditorViewportClient> FWorkspaceEditor::CreateEditorViewportClient() const
+{
+	FAdvancedPreviewScene* NewPreviewScene = new FAdvancedPreviewScene(FPreviewScene::ConstructionValues());
+	
+	StaticCastSharedPtr<FAssetEditorModeManager>(EditorModeManager)->SetPreviewScene(NewPreviewScene);
+
+	TSharedPtr<FEditorViewportClient> EditorViewportClient = MakeShared<FWorkspaceAssetViewportClient>(EditorModeManager.Get(), NewPreviewScene);
+	EditorViewportClient->ViewportType = LVT_Perspective;
+	EditorViewportClient->bSetListenerPosition = false;
+	EditorViewportClient->SetViewLocation(FVector(0.0f, 400.0f, 200.0f));
+	EditorViewportClient->SetViewRotation(EditorViewportDefs::DefaultPerspectiveViewRotation);
+	EditorViewportClient->SetRealtime(true);
+	EditorViewportClient->SetViewMode(VMI_Lit);
+	EditorViewportClient->ToggleOrbitCamera(true);
+
+	return EditorViewportClient;
+}
+
+void FWorkspaceEditor::SetSceneDescription(TObjectPtr<UWorkspaceViewportSceneDescription> InSceneDescription)
+{
+	SceneDescription = InSceneDescription;
+
+	SceneDescription->GetOnConfigChanged().AddSP(this, &FWorkspaceEditor::ReinitViewportScene);
+}
+	
+UWorkspaceViewportSceneDescription* FWorkspaceEditor::GetSceneDescription() const
+{
+	return SceneDescription;
+}
+	
 }
 
 #undef LOCTEXT_NAMESPACE
diff --git a/Engine/Plugins/Experimental/Workspace/Source/WorkspaceEditor/Private/WorkspaceEditor.h b/Engine/Plugins/Experimental/Workspace/Source/WorkspaceEditor/Private/WorkspaceEditor.h
--- a/Engine/Plugins/Experimental/Workspace/Source/WorkspaceEditor/Private/WorkspaceEditor.h
+++ b/Engine/Plugins/Experimental/Workspace/Source/WorkspaceEditor/Private/WorkspaceEditor.h
@@ -23,12 +23,12 @@ class IStructureDetailsView;
 
 namespace UE::Workspace
 {
-
 struct FGraphDocumentSummoner;
 struct FWorkspaceTabSummoner;
 struct FAssetDocumentSummoner;
 class FWorkspaceEditorModule;
 class SWorkspaceView;
+class IWorkspaceViewportController;
 
 namespace WorkspaceModes
 {
@@ -73,7 +73,9 @@ private:
 	virtual void RegisterToolbar() override;
 	virtual bool ShouldReopenEditorForSavedAsset(const UObject* Asset) const override;
 	virtual void RemoveEditingObject(UObject* Object) override;
-
+	virtual AssetEditorViewportFactoryFunction GetViewportDelegate() override;
+	virtual TSharedPtr<FEditorViewportClient> CreateEditorViewportClient() const override;
+	
 	// FAssetEditorToolkit interface
 	virtual void GetSaveableObjects(TArray<UObject*>& OutObjects) const override;
 	virtual bool CanSaveAsset() const override;
@@ -85,10 +87,7 @@ private:
 	virtual void OnToolkitHostingStarted(const TSharedRef<IToolkit>& Toolkit) override;
 
 	// FGCObject interface
-	virtual void AddReferencedObjects(FReferenceCollector& Collector) override
-	{
-		Collector.AddReferencedObject(Workspace);
-	}
+	virtual void AddReferencedObjects(FReferenceCollector& Collector) override;
 	virtual FString GetReferencerName() const override
 	{
 		return TEXT("FWorkspaceEditor");
@@ -113,6 +112,10 @@ private:
 	virtual UObject* GetWorkspaceAsset() const override;
 	virtual FString GetPackageName() const override;
 	virtual const FWorkspaceDocument& GetFocusedWorkspaceDocument() const override;
+	virtual void SetSceneDescription(TObjectPtr<UWorkspaceViewportSceneDescription> InSceneDescription) override;
+	virtual UWorkspaceViewportSceneDescription* GetSceneDescription() const override;
+	
+	void ReinitViewportScene();
 
 	void HandleOutlinerSelectionChanged(TConstArrayView<FWorkspaceOutlinerItemExport> InExports);
 
@@ -173,6 +176,9 @@ private:
 	TSharedPtr<FWorkspaceItem> EditorMenuCategory;
 
 	TOptional<FWorkspaceDocument> LastDocument;
+
+	TUniquePtr<IWorkspaceViewportController> ActiveViewportController;
+	TObjectPtr<UWorkspaceViewportSceneDescription> SceneDescription;
 };
 
 }
\ No newline at end of file
diff --git a/Engine/Plugins/Experimental/Workspace/Source/WorkspaceEditor/Private/WorkspaceEditorModule.cpp b/Engine/Plugins/Experimental/Workspace/Source/WorkspaceEditor/Private/WorkspaceEditorModule.cpp
--- a/Engine/Plugins/Experimental/Workspace/Source/WorkspaceEditor/Private/WorkspaceEditorModule.cpp
+++ b/Engine/Plugins/Experimental/Workspace/Source/WorkspaceEditor/Private/WorkspaceEditorModule.cpp
@@ -357,6 +357,26 @@ TSharedPtr<IWorkspaceOutlinerItemDetails> FWorkspaceEditorModule::GetOutlinerIte
 	TSharedPtr<IWorkspaceOutlinerItemDetails>* FoundDetails = OutlinerItemDetails.Find(InItemDetailsId);
 	return FoundDetails ? *FoundDetails : nullptr;
 }
+
+void FWorkspaceEditorModule::RegisterViewportControllerFactory(const UClass* InClassPath, const FWorkspaceViewportControllerFactory InControllerFactory)
+{
+	ViewportControllerFactories.Add(InClassPath->GetFName(), InControllerFactory);	
+}
+	
+void FWorkspaceEditorModule::UnregisterViewportControllerFactory(const UClass* InClassPath)
+{
+	ViewportControllerFactories.Remove(InClassPath->GetFName());
+}
+	
+TUniquePtr<IWorkspaceViewportController> FWorkspaceEditorModule::CreateViewportController(const UClass* InClassPath) const
+{
+	if (const FWorkspaceViewportControllerFactory* ControllerFactory = ViewportControllerFactories.Find(InClassPath->GetFName()))
+	{
+		return (*ControllerFactory)();
+	}
+	return nullptr;
+}
+
 }
 
 IMPLEMENT_MODULE(UE::Workspace::FWorkspaceEditorModule, WorkspaceEditor);
diff --git a/Engine/Plugins/Experimental/Workspace/Source/WorkspaceEditor/Private/WorkspaceEditorModule.h b/Engine/Plugins/Experimental/Workspace/Source/WorkspaceEditor/Private/WorkspaceEditorModule.h
--- a/Engine/Plugins/Experimental/Workspace/Source/WorkspaceEditor/Private/WorkspaceEditorModule.h
+++ b/Engine/Plugins/Experimental/Workspace/Source/WorkspaceEditor/Private/WorkspaceEditorModule.h
@@ -40,6 +40,12 @@ private:
 	virtual FOnExtendTabs& OnExtendTabs() override { return ExtendTabsForEditor; }
 	virtual FOnExtendToolMenuContext& OnExtendToolMenuContext() override { return ExtendToolMenuContext; }
 	virtual TSharedPtr<IWorkspacePicker> CreateWorkspacePicker(const IWorkspacePicker::FConfig& Config) const override;
+	virtual void RegisterViewportControllerFactory(const UClass* InClassPath, const FWorkspaceViewportControllerFactory InControllerFactory) override;
+	virtual void UnregisterViewportControllerFactory(const UClass* InClassPath) override;
+
+	// Create a new viewport controller for a given asset class if a controller has been registered
+	// * Can return nullptr
+	TUniquePtr<IWorkspaceViewportController> CreateViewportController(const UClass* InClassPath) const;
 
 	// Find an existing registered object document type. Note this redirects based on FObjectDocumentArgs::OnRedirectWorkspaceContext
 	const FObjectDocumentArgs* FindObjectDocumentType(const FWorkspaceDocument& Document) const;
@@ -58,6 +64,8 @@ private:
 
 	TMap<FTopLevelAssetPath, FObjectDocumentArgs> ObjectDocumentArgs;
 
+	TMap<FName, FWorkspaceViewportControllerFactory> ViewportControllerFactories;
+	
 	TMap<FName, TSet<FTopLevelAssetPath>> DocumentAreaMap;
 
 	TMap<FTopLevelAssetPath, FDocumentSubObjectArgs> DocumentSubObjectArgs; 
diff --git a/Engine/Plugins/Experimental/Workspace/Source/WorkspaceEditor/Private/WorkspaceViewportSceneDescription.cpp b/Engine/Plugins/Experimental/Workspace/Source/WorkspaceEditor/Private/WorkspaceViewportSceneDescription.cpp
new file mode 100644
--- /dev/null
+++ b/Engine/Plugins/Experimental/Workspace/Source/WorkspaceEditor/Private/WorkspaceViewportSceneDescription.cpp
@@ -0,0 +1,13 @@
+// Copyright Epic Games, Inc. All Rights Reserved.
+
+#include "WorkspaceViewportSceneDescription.h"
+
+FSimpleMulticastDelegate& UWorkspaceViewportSceneDescription::GetOnConfigChanged()
+{
+	return OnConfigChanged;
+}
+
+void UWorkspaceViewportSceneDescription::BroadcastOnConfigChanged()
+{
+	OnConfigChanged.Broadcast();
+}
\ No newline at end of file
diff --git a/Engine/Plugins/Experimental/Workspace/Source/WorkspaceEditor/Public/IWorkspaceEditor.h b/Engine/Plugins/Experimental/Workspace/Source/WorkspaceEditor/Public/IWorkspaceEditor.h
--- a/Engine/Plugins/Experimental/Workspace/Source/WorkspaceEditor/Public/IWorkspaceEditor.h
+++ b/Engine/Plugins/Experimental/Workspace/Source/WorkspaceEditor/Public/IWorkspaceEditor.h
@@ -6,6 +6,7 @@
 
 #define UE_API WORKSPACEEDITOR_API
 
+class UWorkspaceViewportSceneDescription;
 class UWorkspaceAssetEditor;
 class UWorkspaceSchema;
 class UWorkspace;
@@ -100,6 +101,12 @@ public:
 
 	// Name of EditorToolkit, used for testing against FAssetEditorToolkit::GetToolkitFName 
 	static FName GetWorkspaceEditorToolkitName() { return FName("WorkspaceEditor"); }
+
+	// Set a new scene description config
+	virtual void SetSceneDescription(TObjectPtr<UWorkspaceViewportSceneDescription> InSceneDescription) = 0;
+
+	// Get the current scene description config
+	virtual UWorkspaceViewportSceneDescription* GetSceneDescription() const = 0;
 };
 
 }
diff --git a/Engine/Plugins/Experimental/Workspace/Source/WorkspaceEditor/Public/IWorkspaceEditorModule.h b/Engine/Plugins/Experimental/Workspace/Source/WorkspaceEditor/Public/IWorkspaceEditorModule.h
--- a/Engine/Plugins/Experimental/Workspace/Source/WorkspaceEditor/Public/IWorkspaceEditorModule.h
+++ b/Engine/Plugins/Experimental/Workspace/Source/WorkspaceEditor/Public/IWorkspaceEditorModule.h
@@ -11,6 +11,7 @@
 #include "Modules/ModuleInterface.h"
 #include "IWorkspaceOutlinerItemDetails.h"
 #include "IWorkspacePicker.h"
+#include "WorkspaceViewportController.h"
 
 #define UE_API WORKSPACEEDITOR_API
 
@@ -21,6 +22,7 @@ class FTabManager;
 class FWorkflowAllowedTabSet;
 class IDetailCustomization;
 class FLayoutExtender;
+class FAdvancedPreviewScene;
 
 struct FSlateBrush;
 struct FTopLevelAssetPath;
@@ -271,6 +273,13 @@ public:
 	virtual FOnExtendToolMenuContext& OnExtendToolMenuContext() = 0;
 
 	virtual TSharedPtr<IWorkspacePicker> CreateWorkspacePicker(const IWorkspacePicker::FConfig& Config) const = 0;
+
+	// Register a viewport controller factory for a given asset class
+	// * Viewport controllers are created per viewport
+	virtual void RegisterViewportControllerFactory(const UClass* InClass, const FWorkspaceViewportControllerFactory InControllerFactory) = 0;
+
+	// Unregister the viewport controller factory for a given asset class
+	virtual void UnregisterViewportControllerFactory(const UClass* InClass) = 0;
 };
 
 }
diff --git a/Engine/Plugins/Experimental/Workspace/Source/WorkspaceEditor/Public/WorkspaceSchema.h b/Engine/Plugins/Experimental/Workspace/Source/WorkspaceEditor/Public/WorkspaceSchema.h
--- a/Engine/Plugins/Experimental/Workspace/Source/WorkspaceEditor/Public/WorkspaceSchema.h
+++ b/Engine/Plugins/Experimental/Workspace/Source/WorkspaceEditor/Public/WorkspaceSchema.h
@@ -30,4 +30,7 @@ public:
 
 	// Called after loading workspace state. InWorkspaceState is an instanced struct that holds user-defined persistent workspace state. Struct is not guaranteed to be valid.
 	virtual void OnLoadWorkspaceState(TSharedRef<UE::Workspace::IWorkspaceEditor> InWorkspaceEditor, const FInstancedStruct& InWorkspaceState) const {}
+
+	// Whether workspaces using this schema can spawn viewport tabs.
+	virtual bool SupportsViewport() const { return false; };
 };
diff --git a/Engine/Plugins/Experimental/Workspace/Source/WorkspaceEditor/Public/WorkspaceViewportController.h b/Engine/Plugins/Experimental/Workspace/Source/WorkspaceEditor/Public/WorkspaceViewportController.h
new file mode 100644
--- /dev/null
+++ b/Engine/Plugins/Experimental/Workspace/Source/WorkspaceEditor/Public/WorkspaceViewportController.h
@@ -0,0 +1,39 @@
+// Copyright Epic Games, Inc. All Rights Reserved.
+
+#pragma once
+
+#include "CoreMinimal.h"
+#include "WorkspaceViewportSceneDescription.h"
+
+class FAdvancedPreviewScene;
+class UWorkspaceViewportSceneDescription;
+
+namespace UE::Workspace
+{
+	class IWorkspaceViewportController
+	{
+	public:
+		virtual ~IWorkspaceViewportController() = default;
+
+		class FViewportContext : public FGCObject
+		{
+		public:
+			FAdvancedPreviewScene* PreviewScene;
+			TObjectPtr<UObject> OutlinerObject;
+			TObjectPtr<UWorkspaceViewportSceneDescription> SceneDescription;
+			
+			// FGCObject interface
+			virtual FString GetReferencerName() const override { return TEXT("IWorkspaceViewportController::FViewportContext"); }
+			virtual void AddReferencedObjects(FReferenceCollector& Collector) override
+			{
+				Collector.AddReferencedObject(OutlinerObject);
+				Collector.AddReferencedObject(SceneDescription);
+			}
+		};
+		
+		virtual void OnEnter(const FViewportContext& InContext) = 0;
+		virtual void OnExit() = 0;
+	};
+
+	using FWorkspaceViewportControllerFactory = TFunctionRef<TUniquePtr<IWorkspaceViewportController>()>;
+}
diff --git a/Engine/Plugins/Experimental/Workspace/Source/WorkspaceEditor/Public/WorkspaceViewportSceneDescription.h b/Engine/Plugins/Experimental/Workspace/Source/WorkspaceEditor/Public/WorkspaceViewportSceneDescription.h
new file mode 100644
--- /dev/null
+++ b/Engine/Plugins/Experimental/Workspace/Source/WorkspaceEditor/Public/WorkspaceViewportSceneDescription.h
@@ -0,0 +1,22 @@
+// Copyright Epic Games, Inc. All Rights Reserved.
+
+#pragma once
+
+#include "CoreMinimal.h"
+
+#include "WorkspaceViewportSceneDescription.generated.h"
+
+UCLASS()
+class WORKSPACEEDITOR_API UWorkspaceViewportSceneDescription : public UObject
+{
+	GENERATED_BODY()
+
+public:
+	FSimpleMulticastDelegate& GetOnConfigChanged();
+
+protected:
+	void BroadcastOnConfigChanged();
+
+private:
+	FSimpleMulticastDelegate OnConfigChanged;
+};
\ No newline at end of file
diff --git a/Engine/Plugins/Experimental/Workspace/Source/WorkspaceEditor/WorkspaceEditor.Build.cs b/Engine/Plugins/Experimental/Workspace/Source/WorkspaceEditor/WorkspaceEditor.Build.cs
--- a/Engine/Plugins/Experimental/Workspace/Source/WorkspaceEditor/WorkspaceEditor.Build.cs
+++ b/Engine/Plugins/Experimental/Workspace/Source/WorkspaceEditor/WorkspaceEditor.Build.cs
@@ -46,7 +46,8 @@ namespace UnrealBuildTool.Rules
 					"SceneOutliner",
 					"EditorInteractiveToolsFramework",
 					"InteractiveToolsFramework",
-					"EditorFramework"
+					"EditorFramework",
+					"AdvancedPreviewScene"
 				}
 			);
 		}
