Verify that users can add an AnimNext Animation Graph in an AnimNext Workspace asset
Automated QMetry Test Case: UE-TC-10700
Misc code cleanup

diff --git a/Engine/Plugins/Experimental/Workspace/Source/WorkspaceEditor/Private/Workspace.h b/Engine/Plugins/Experimental/Workspace/Source/WorkspaceEditor/Private/Workspace.h
--- a/Engine/Plugins/Experimental/Workspace/Source/WorkspaceEditor/Private/Workspace.h
+++ b/Engine/Plugins/Experimental/Workspace/Source/WorkspaceEditor/Private/Workspace.h
@@ -84,6 +84,8 @@ class UWorkspace : public UObject
 	friend class UWorkspaceState;
 
 	friend class UE::UAF::Editor::FAssetWizard;
+	
+	friend struct FUAFWorkspaceTests;
 
 	// Adds an asset to the workspace
 	// @return true if the asset was added
