Fix for crash when creating a new UAF Workspace instead of picking an existing one when prompted to choose a workspace after opening a newly created UAF Graph

diff --git a/Engine/Plugins/Experimental/Workspace/Source/WorkspaceEditor/Private/WorkspaceEditorModule.cpp b/Engine/Plugins/Experimental/Workspace/Source/WorkspaceEditor/Private/WorkspaceEditorModule.cpp
--- a/Engine/Plugins/Experimental/Workspace/Source/WorkspaceEditor/Private/WorkspaceEditorModule.cpp
+++ b/Engine/Plugins/Experimental/Workspace/Source/WorkspaceEditor/Private/WorkspaceEditorModule.cpp
@@ -282,7 +282,8 @@ bool FWorkspaceEditorModule::GetExportedAssetsForWorkspace(const FAssetData& InW
 		TSharedRef<SWorkspacePicker> WorkspacePicker = SNew(SWorkspacePicker)
 			.WorkspaceAssets(RelevantWorkspaceAssets)
 			.OnExistingWorkspaceSelected_Lambda(HandleExistingWorkspace)
-			.OnNewWorkspaceCreated_Lambda(HandleNewWorkspace);
+			.OnNewWorkspaceCreated_Lambda(HandleNewWorkspace)
+			.WorkspaceFactoryClass(WorkSpaceFactoryClass.Get());
 
 		WorkspacePicker->ShowModal();
 	}
