Add pinning and toolbar to workspace viewport and remove scene description tab

diff --git a/Engine/Plugins/Experimental/UAF/UAF/Source/UAFEditor/Private/AnimNextEditorModule.cpp b/Engine/Plugins/Experimental/UAF/UAF/Source/UAFEditor/Private/AnimNextEditorModule.cpp
--- a/Engine/Plugins/Experimental/UAF/UAF/Source/UAFEditor/Private/AnimNextEditorModule.cpp
+++ b/Engine/Plugins/Experimental/UAF/UAF/Source/UAFEditor/Private/AnimNextEditorModule.cpp
@@ -74,7 +74,6 @@
 #include "UAFStyle.h"
 #include "Common/AnimNextAssetFindReplaceVariables.h"
 #include "Toolkits/AssetEditorToolkitMenuContext.h"
-#include "UAF/Viewport/SceneDescriptionTab.h"
 #include "UAF/Viewport/SystemViewportController.h"
 #include "Variables/Outliner/VariablesOutlinerCommands.h"
 #include "Wizard/AssetWizard.h"
@@ -183,10 +182,6 @@ void FAnimNextEditorModule::StartupModule()
 			TSharedRef<FAnimNextFindAndReplaceTabSummoner> FindAndReplaceTabSummoner = MakeShared<FAnimNextFindAndReplaceTabSummoner>(InEditorPtr);
 			TabFactories.RegisterFactory(FindAndReplaceTabSummoner);
 			FindAndReplaceTabSummoner->RegisterTabSpawner(InTabManager, nullptr);
-
-			TSharedRef<FUAFViewportSceneDescriptionTabSummoner> ViewportSceneDescriptionTabSummoner = MakeShared<FUAFViewportSceneDescriptionTabSummoner>(InEditorPtr);
-			TabFactories.RegisterFactory(ViewportSceneDescriptionTabSummoner);
-			ViewportSceneDescriptionTabSummoner->RegisterTabSpawner(InTabManager, nullptr);
 		});
 
 	WorkspaceEditorModule.OnExtendTabs().AddLambda([](FLayoutExtender& InLayoutExtender, TSharedPtr<UE::Workspace::IWorkspaceEditor> InEditorPtr)
diff --git a/Engine/Plugins/Experimental/UAF/UAF/Source/UAFEditor/Private/UAF/Viewport/SViewportSceneDescription.cpp b/Engine/Plugins/Experimental/UAF/UAF/Source/UAFEditor/Private/UAF/Viewport/SViewportSceneDescription.cpp
deleted file mode 100644
diff --git a/Engine/Plugins/Experimental/UAF/UAF/Source/UAFEditor/Private/UAF/Viewport/SViewportSceneDescription.h b/Engine/Plugins/Experimental/UAF/UAF/Source/UAFEditor/Private/UAF/Viewport/SViewportSceneDescription.h
deleted file mode 100644
diff --git a/Engine/Plugins/Experimental/UAF/UAF/Source/UAFEditor/Private/UAF/Viewport/SceneDescriptionTab.cpp b/Engine/Plugins/Experimental/UAF/UAF/Source/UAFEditor/Private/UAF/Viewport/SceneDescriptionTab.cpp
deleted file mode 100644
diff --git a/Engine/Plugins/Experimental/UAF/UAF/Source/UAFEditor/Private/UAF/Viewport/SceneDescriptionTab.h b/Engine/Plugins/Experimental/UAF/UAF/Source/UAFEditor/Private/UAF/Viewport/SceneDescriptionTab.h
deleted file mode 100644
diff --git a/Engine/Plugins/Experimental/UAF/UAF/Source/UAFEditor/Private/Workspace/AnimNextWorkspaceSchema.h b/Engine/Plugins/Experimental/UAF/UAF/Source/UAFEditor/Private/Workspace/AnimNextWorkspaceSchema.h
--- a/Engine/Plugins/Experimental/UAF/UAF/Source/UAFEditor/Private/Workspace/AnimNextWorkspaceSchema.h
+++ b/Engine/Plugins/Experimental/UAF/UAF/Source/UAFEditor/Private/Workspace/AnimNextWorkspaceSchema.h
@@ -3,6 +3,8 @@
 #pragma once
 
 #include "WorkspaceSchema.h"
+#include "UAF/Viewport/ViewportSceneDescription.h"
+
 #include "AnimNextWorkspaceSchema.generated.h"
 
 // Workspace schema allowing all asset types
@@ -17,4 +19,5 @@ class UAnimNextWorkspaceSchema : public UWorkspaceSchema
 	virtual void OnSaveWorkspaceState(TSharedRef<UE::Workspace::IWorkspaceEditor> InWorkspaceEditor, FInstancedStruct& OutWorkspaceState) const override;
 	virtual void OnLoadWorkspaceState(TSharedRef<UE::Workspace::IWorkspaceEditor> InWorkspaceEditor, const FInstancedStruct& InWorkspaceState) const override;
 	virtual bool SupportsViewport() const { return true; };
+	virtual TObjectPtr<UWorkspaceViewportSceneDescription> CreateSceneDescription() const { return NewObject<UUAFViewportSceneDescription>(); }
 };
\ No newline at end of file
diff --git a/Engine/Plugins/Experimental/Workspace/Source/WorkspaceEditor/Private/SWorkspaceViewport.cpp b/Engine/Plugins/Experimental/Workspace/Source/WorkspaceEditor/Private/SWorkspaceViewport.cpp
new file mode 100644
--- /dev/null
+++ b/Engine/Plugins/Experimental/Workspace/Source/WorkspaceEditor/Private/SWorkspaceViewport.cpp
@@ -0,0 +1,208 @@
+// Copyright Epic Games, Inc. All Rights Reserved.
+
+#include "SWorkspaceViewport.h"
+
+#include "AdvancedPreviewScene.h"
+#include "DetailsViewArgs.h"
+#include "PropertyEditorModule.h"
+#include "ToolMenus.h"
+#include "Widgets/Input/SCheckBox.h"
+#include "Widgets/Images/SImage.h"
+#include "WorkspaceAssetViewportClient.h"
+#include "WorkspaceViewportMenuContext.h"
+#include "WorkspaceViewportSceneDescription.h"
+#include "ViewportToolbar/UnrealEdViewportToolbar.h"
+
+#define LOCTEXT_NAMESPACE "UE::Workspace::SWorkspaceViewport"
+
+namespace UE::Workspace
+{
+	void SWorkspaceViewport::Construct(const FArguments& InArgs)
+	{
+		Client = InArgs._ViewportClient;
+		AssetEditorToolkit = InArgs._AssetEditorToolkit;
+		SceneDescription = InArgs._SceneDescription;
+		
+		PreviewAssetPath = InArgs._PreviewAssetPath;
+		bIsPinned = InArgs._bIsPinned;
+		OnPinnedClicked = InArgs._OnPinnedClicked;
+		
+		SEditorViewport::Construct(SEditorViewport::FArguments());
+	}
+
+	TSharedRef<FEditorViewportClient> SWorkspaceViewport::MakeEditorViewportClient()
+	{
+		return Client.ToSharedRef();
+	}
+
+	TSharedPtr<SWidget> SWorkspaceViewport::BuildViewportToolbar()
+	{
+		const FName ViewportToolbarName = "WorkspaceEditor.ViewportToolbar";
+
+		if (!UToolMenus::Get()->IsMenuRegistered(ViewportToolbarName))
+		{
+			UToolMenu* const ViewportToolbarMenu = UToolMenus::Get()->RegisterMenu(
+				ViewportToolbarName, NAME_None /* parent */, EMultiBoxType::SlimHorizontalToolBar
+			);
+			ViewportToolbarMenu->StyleName = "ViewportToolbar";
+
+			ViewportToolbarMenu->AddDynamicSection("Left", FNewToolMenuDelegate::CreateLambda([](UToolMenu* InMenu)
+				{
+					if (UWorkspaceViewportMenuContext* Context = InMenu->FindContext<UWorkspaceViewportMenuContext>())
+					{
+						InMenu->AddMenuEntry(NAME_None, FToolMenuEntry::InitWidget(
+							"PinPreview",
+							SNew(SBox)
+							[
+								SNew(SCheckBox)
+								.Style(FAppStyle::Get(), "DetailsView.SectionButton")
+								[
+									SNew(SHorizontalBox)
+									+ SHorizontalBox::Slot()
+									.Padding(FMargin(0.0f, 0.0f, 4.0f, 0.0f))
+									.AutoWidth()
+									[
+										SNew(SImage)
+										.Image_Lambda([Context]()
+										{
+											const bool bViewportPinned = Context->bIsPinned.Get();
+											return FAppStyle::Get().GetBrush(bViewportPinned ? "Icons.Lock" : "Icons.Unlock");
+										})
+									]
+									+ SHorizontalBox::Slot()
+									.AutoWidth()
+									[
+										SNew(STextBlock)
+										.ShadowOffset(FVector2D::UnitVector)
+										.ShadowColorAndOpacity(FLinearColor(0.0f, 0.0f, 0.0f, 0.5f))
+										.Text_Lambda([Context]()
+										{
+											return FText::Format(LOCTEXT("OverlayText", "Previewing {0}"), FText::FromString(Context->PreviewAssetPath.Get().GetAssetName()));
+										})
+									]
+								]
+								.IsChecked_Lambda([Context]()
+								{
+									return Context->bIsPinned.Get() ? ECheckBoxState::Checked : ECheckBoxState::Unchecked;
+								})
+								.OnCheckStateChanged_Lambda([Context](ECheckBoxState InState)
+								{
+									Context->OnPinnedClicked.Execute();
+								})
+							],
+							FText::GetEmpty()));
+					
+						InMenu->AddMenuEntry(NAME_None, FToolMenuEntry::InitSubMenu(
+							"SceneDescription",
+							LOCTEXT("SceneDescription_Label", "Scene Settings"),
+							LOCTEXT("SceneDescription_ToolTip", "Options to control the preview scene"),
+							FNewToolMenuDelegate::CreateLambda([Context](UToolMenu* Submenu)
+								{
+									FPropertyEditorModule& PropertyEditorModule = FModuleManager::LoadModuleChecked<FPropertyEditorModule>("PropertyEditor");
+													
+									FDetailsViewArgs DetailsViewArgs;
+									DetailsViewArgs.bAllowSearch = false;
+									DetailsViewArgs.bAllowFavoriteSystem = false;
+									DetailsViewArgs.bShowOptions = false;
+									DetailsViewArgs.bShowObjectLabel = false;
+									DetailsViewArgs.bHideSelectionTip = true;
+
+									const auto DetailsView = PropertyEditorModule.CreateDetailView(DetailsViewArgs);
+									DetailsView->SetObject(Context->SceneDescription);
+							
+									Submenu->AddMenuEntry(NAME_None, FToolMenuEntry::InitWidget(
+										"Details",
+										DetailsView,
+										FText::GetEmpty(),
+										true,
+										false,
+										true
+									));
+								}),
+							false,
+							FSlateIcon(FAppStyle::GetAppStyleSetName(), "AssetEditor.PreviewSceneSettings")));
+					}
+				}));
+
+			{
+				FToolMenuSection& RightSection = ViewportToolbarMenu->AddSection("Right");
+				RightSection.Alignment = EToolMenuSectionAlign::Last;
+	
+				// Add the "Camera" submenu.
+				RightSection.AddEntry(UE::UnrealEd::CreateCameraSubmenu(UE::UnrealEd::FViewportCameraMenuOptions().ShowLensControls()));
+	
+				// Add the "View Modes" sub menu.
+				{
+					// Stay backward-compatible with the old viewport toolbar.
+					{
+						const FName ParentSubmenuName = "UnrealEd.ViewportToolbar.View";
+						// Create our parent menu.
+						if (!UToolMenus::Get()->IsMenuRegistered(ParentSubmenuName))
+						{
+							UToolMenus::Get()->RegisterMenu(ParentSubmenuName);
+						}
+	
+						// Register our ToolMenu here first, before we create the submenu, so we can set our parent.
+						UToolMenus::Get()->RegisterMenu("WaterWavesEditor.ViewportToolbar.ViewModes", ParentSubmenuName);
+					}
+	
+					RightSection.AddEntry(UE::UnrealEd::CreateViewModesSubmenu());
+				}
+	
+				RightSection.AddEntry(UE::UnrealEd::CreateDefaultShowSubmenu());
+				RightSection.AddEntry(UE::UnrealEd::CreatePerformanceAndScalabilitySubmenu());
+				RightSection.AddEntry(UE::UnrealEd::CreateAssetViewerProfileSubmenu());
+			}
+		}
+
+		FToolMenuContext ViewportToolbarContext;
+		{
+			UWorkspaceViewportMenuContext* WorkspaceViewportContext = ViewportToolbarContext.FindContext<UWorkspaceViewportMenuContext>();
+			if (!WorkspaceViewportContext)
+			{
+				WorkspaceViewportContext = NewObject<UWorkspaceViewportMenuContext>();
+				ViewportToolbarContext.AddObject(WorkspaceViewportContext);
+			}
+
+			WorkspaceViewportContext->bIsPinned = bIsPinned;
+			WorkspaceViewportContext->PreviewAssetPath = PreviewAssetPath;
+			WorkspaceViewportContext->OnPinnedClicked = OnPinnedClicked;
+			WorkspaceViewportContext->SceneDescription = SceneDescription;
+			
+			ViewportToolbarContext.AppendCommandList(GetCommandList());
+
+			// Add the UnrealEd viewport toolbar context.
+			{
+				UUnrealEdViewportToolbarContext* const ContextObject = NewObject<UUnrealEdViewportToolbarContext>();
+				ContextObject->Viewport = SharedThis(this);
+				ContextObject->AssetEditorToolkit = AssetEditorToolkit;
+
+				ViewportToolbarContext.AddObject(ContextObject);
+			}
+		}
+
+		return UToolMenus::Get()->GenerateWidget(ViewportToolbarName, ViewportToolbarContext);
+	}
+
+	TSharedRef<SEditorViewport> SWorkspaceViewport::GetViewportWidget()
+	{
+		return SharedThis(this);
+	}
+
+	TSharedPtr<FExtender> SWorkspaceViewport::GetExtenders() const
+	{
+		TSharedPtr<FExtender> Result(MakeShareable(new FExtender));
+		return Result;
+	}
+
+	void SWorkspaceViewport::OnFloatingButtonClicked()
+	{
+	}
+	
+	void SWorkspaceViewport::AddReferencedObjects(FReferenceCollector& Collector)
+	{
+		Collector.AddReferencedObject(SceneDescription);
+	}
+}
+
+#undef LOCTEXT_NAMESPACE
\ No newline at end of file
diff --git a/Engine/Plugins/Experimental/Workspace/Source/WorkspaceEditor/Private/SWorkspaceViewport.h b/Engine/Plugins/Experimental/Workspace/Source/WorkspaceEditor/Private/SWorkspaceViewport.h
new file mode 100644
--- /dev/null
+++ b/Engine/Plugins/Experimental/Workspace/Source/WorkspaceEditor/Private/SWorkspaceViewport.h
@@ -0,0 +1,50 @@
+// Copyright Epic Games, Inc. All Rights Reserved.
+
+#pragma once
+
+#include "SAssetEditorViewport.h"
+#include "SCommonEditorViewportToolbarBase.h"
+
+class UWorkspaceViewportSceneDescription;
+
+namespace UE::Workspace
+{
+	class SWorkspaceViewport : public SAssetEditorViewport, public ICommonEditorViewportToolbarInfoProvider, public FGCObject
+	{
+	public:
+		SLATE_BEGIN_ARGS(SWorkspaceViewport){}
+			SLATE_ARGUMENT(TWeakPtr<FAssetEditorToolkit>, AssetEditorToolkit)
+			SLATE_ARGUMENT(TSharedPtr<FEditorViewportClient>, ViewportClient)
+			SLATE_ARGUMENT(UWorkspaceViewportSceneDescription*, SceneDescription)
+			SLATE_ATTRIBUTE(FSoftObjectPath, PreviewAssetPath)
+			SLATE_ATTRIBUTE(bool, bIsPinned)
+			SLATE_EVENT(FSimpleDelegate, OnPinnedClicked)
+		SLATE_END_ARGS()
+		
+		TWeakPtr<FAssetEditorToolkit> AssetEditorToolkit = nullptr;
+
+		void Construct(const FArguments& InArgs);
+		
+		// SEditorViewport interface
+		virtual TSharedPtr<SWidget> BuildViewportToolbar() override;
+		virtual TSharedRef<FEditorViewportClient> MakeEditorViewportClient() override;
+		
+		// ICommonEditorViewportToolbarInfoProvider interface
+		virtual TSharedRef<SEditorViewport> GetViewportWidget() override;
+		virtual TSharedPtr<FExtender> GetExtenders() const override;
+		virtual void OnFloatingButtonClicked() override;
+
+		virtual void AddReferencedObjects( FReferenceCollector& Collector ) override;
+		virtual FString GetReferencerName() const override
+		{
+			return TEXT("UE::Workspace::SWorkspaceViewport");
+		}
+
+	private:
+		TAttribute<FSoftObjectPath> PreviewAssetPath;
+		TAttribute<bool> bIsPinned;
+		FSimpleDelegate OnPinnedClicked;
+
+		TObjectPtr<UWorkspaceViewportSceneDescription> SceneDescription;
+	};
+}
diff --git a/Engine/Plugins/Experimental/Workspace/Source/WorkspaceEditor/Private/WorkspaceEditor.cpp b/Engine/Plugins/Experimental/Workspace/Source/WorkspaceEditor/Private/WorkspaceEditor.cpp
--- a/Engine/Plugins/Experimental/Workspace/Source/WorkspaceEditor/Private/WorkspaceEditor.cpp
+++ b/Engine/Plugins/Experimental/Workspace/Source/WorkspaceEditor/Private/WorkspaceEditor.cpp
@@ -27,6 +27,7 @@
 #include "WorkspaceTabPayload.h"
 #include "AdvancedPreviewScene.h"
 #include "ContextObjectStore.h"
+#include "SWorkspaceViewport.h"
 #include "WorkspaceAssetViewportClient.h"
 #include "WorkspaceViewportController.h"
 #include "Viewports.h"
@@ -224,6 +225,9 @@ void FWorkspaceEditor::PostInitAssetEditor()
 	ExtendMenu();
 	ExtendToolbar();
 	RegenerateMenusAndToolbars();
+
+	// Default pin the document that was opened (the root-most asset)
+	ViewportPinnedAssetPath = GetFocusedWorkspaceDocument().Export.GetFirstAssetPath();
 }
 
 void FWorkspaceEditor::OnToolkitHostingFinished(const TSharedRef<IToolkit>& Toolkit)
@@ -552,7 +556,10 @@ void FWorkspaceEditor::SetFocusedDocument(const UE::Workspace::FWorkspaceDocumen
 		LastDocument.Reset();
 	}
 
-	ReinitViewportScene();
+	if (!IsViewportPinned())
+	{
+		ReinitViewportScene();
+	}
 	OnFocusedDocumentChangedDelegate.Broadcast(Document);
 }
 
@@ -614,7 +621,9 @@ void FWorkspaceEditor::ReinitViewportScene()
 		ActiveViewportController = nullptr;
 	}
 
-	const FSoftObjectPath AssetPath = GetFocusedWorkspaceDocument().Export.GetFirstAssetPath();
+	const FSoftObjectPath AssetPath = ViewportPinnedAssetPath.IsValid()
+		? ViewportPinnedAssetPath
+		: GetFocusedWorkspaceDocument().Export.GetFirstAssetPath();
 
 	if (AssetPath.IsValid())
 	{
@@ -888,6 +897,11 @@ bool FWorkspaceEditor::AreAssetEntriesModified() const
 	return false;
 }
 
+bool FWorkspaceEditor::IsViewportPinned() const
+{
+	return ViewportPinnedAssetPath.IsValid();
+}
+
 void FWorkspaceEditor::FindInContentBrowser_Execute()
 {
 	if (IsValid(Workspace))
@@ -977,11 +991,42 @@ void FWorkspaceEditor::RemoveEditingObject(UObject* Object)
 
 AssetEditorViewportFactoryFunction FWorkspaceEditor::GetViewportDelegate()
 {
+	SetSceneDescription(GetSchema()->CreateSceneDescription());
+	
 	return [this](FAssetEditorViewportConstructionArgs InArgs)
 		{
-			TSharedRef<SAssetEditorViewport> Viewport = SNew(SAssetEditorViewport, InArgs)
-				.EditorViewportClient(ViewportClient);
+			TSharedRef<UE::Workspace::SWorkspaceViewport> Viewport = SNew(UE::Workspace::SWorkspaceViewport)
+				.AssetEditorToolkit(SharedThis(this))
+				.ViewportClient(ViewportClient)
+				.SceneDescription(SceneDescription)
+				.bIsPinned_Lambda([this]()
+				{
+					return ViewportPinnedAssetPath.IsValid();
+				})
+				.OnPinnedClicked_Lambda([this]()
+				{
+					if (ViewportPinnedAssetPath.IsValid())
+					{
+						const bool bShouldReinitViewport = ViewportPinnedAssetPath != GetFocusedWorkspaceDocument().Export.GetFirstAssetPath();
+						ViewportPinnedAssetPath.Reset();
 
+						if (bShouldReinitViewport)
+						{
+							ReinitViewportScene();
+						}
+					}
+					else
+					{
+						ViewportPinnedAssetPath = GetFocusedWorkspaceDocument().Export.GetFirstAssetPath();
+					}
+				})
+				.PreviewAssetPath_Lambda([this]()
+				{
+					return ViewportPinnedAssetPath.IsValid()
+						? ViewportPinnedAssetPath
+						: GetFocusedWorkspaceDocument().Export.GetFirstAssetPath();
+				});
+			
 			return Viewport;
 		};
 }
diff --git a/Engine/Plugins/Experimental/Workspace/Source/WorkspaceEditor/Private/WorkspaceEditor.h b/Engine/Plugins/Experimental/Workspace/Source/WorkspaceEditor/Private/WorkspaceEditor.h
--- a/Engine/Plugins/Experimental/Workspace/Source/WorkspaceEditor/Private/WorkspaceEditor.h
+++ b/Engine/Plugins/Experimental/Workspace/Source/WorkspaceEditor/Private/WorkspaceEditor.h
@@ -146,6 +146,8 @@ private:
 	void SaveAssetEntries();
 	bool AreAssetEntriesModified() const;
 
+	bool IsViewportPinned() const;
+
 	/** The asset being edited */
 	TObjectPtr<UWorkspace> Workspace = nullptr;
 
@@ -179,6 +181,7 @@ private:
 
 	TUniquePtr<IWorkspaceViewportController> ActiveViewportController;
 	TObjectPtr<UWorkspaceViewportSceneDescription> SceneDescription;
+	FSoftObjectPath ViewportPinnedAssetPath;
 };
 
 }
\ No newline at end of file
diff --git a/Engine/Plugins/Experimental/Workspace/Source/WorkspaceEditor/Private/WorkspaceViewportMenuContext.h b/Engine/Plugins/Experimental/Workspace/Source/WorkspaceEditor/Private/WorkspaceViewportMenuContext.h
new file mode 100644
--- /dev/null
+++ b/Engine/Plugins/Experimental/Workspace/Source/WorkspaceEditor/Private/WorkspaceViewportMenuContext.h
@@ -0,0 +1,24 @@
+// Copyright Epic Games, Inc. All Rights Reserved.
+
+#pragma once
+
+#include "CoreMinimal.h"
+#include "Misc/Attribute.h"
+
+#include "WorkspaceViewportMenuContext.generated.h"
+
+class UWorkspaceViewportSceneDescription;
+
+UCLASS()
+class UWorkspaceViewportMenuContext : public UObject
+{
+	GENERATED_BODY()
+
+public:
+	TAttribute<bool> bIsPinned;
+	TAttribute<FSoftObjectPath> PreviewAssetPath;
+	FSimpleDelegate OnPinnedClicked;
+
+	UPROPERTY()
+	TObjectPtr<UWorkspaceViewportSceneDescription> SceneDescription;
+};
\ No newline at end of file
diff --git a/Engine/Plugins/Experimental/Workspace/Source/WorkspaceEditor/Public/WorkspaceSchema.h b/Engine/Plugins/Experimental/Workspace/Source/WorkspaceEditor/Public/WorkspaceSchema.h
--- a/Engine/Plugins/Experimental/Workspace/Source/WorkspaceEditor/Public/WorkspaceSchema.h
+++ b/Engine/Plugins/Experimental/Workspace/Source/WorkspaceEditor/Public/WorkspaceSchema.h
@@ -2,6 +2,8 @@
 
 #pragma once
 
+#include "WorkspaceViewportSceneDescription.h"
+
 #include "WorkspaceSchema.generated.h"
 
 class UWorkspace;
@@ -33,4 +35,7 @@ public:
 
 	// Whether workspaces using this schema can spawn viewport tabs.
 	virtual bool SupportsViewport() const { return false; };
+
+	// Create a new object containing the required data to construct a scene for this type of workspace
+	virtual TObjectPtr<UWorkspaceViewportSceneDescription> CreateSceneDescription() const { return NewObject<UWorkspaceViewportSceneDescription>(); }
 };
