UAF: Add comment creation via C Hotkey

diff --git a/Engine/Plugins/Experimental/UAF/UAF/Source/UAFEditor/Private/AnimNextEditorModule.cpp b/Engine/Plugins/Experimental/UAF/UAF/Source/UAFEditor/Private/AnimNextEditorModule.cpp
--- a/Engine/Plugins/Experimental/UAF/UAF/Source/UAFEditor/Private/AnimNextEditorModule.cpp
+++ b/Engine/Plugins/Experimental/UAF/UAF/Source/UAFEditor/Private/AnimNextEditorModule.cpp
@@ -1374,7 +1374,20 @@ void FAnimNextEditorModule::RegisterWorkspaceDocumentTypes(Workspace::IWorkspace
 			}
 		}
 	});
-
+	GraphArgs.OnCreateComment = Workspace::FOnCreateComment::CreateLambda([](const Workspace::FWorkspaceEditorContext& InContext, const FVector2f& GraphLocation)
+		{
+			if (UAnimNextEdGraph* EdGraph = InContext.Document.GetTypedObject<UAnimNextEdGraph>())
+			{
+				if (URigVMController* Controller = EdGraph->GetController())
+				{
+					Controller->OpenUndoBracket(TEXT("Create Comment"));
+					TArray<UEdGraphPin*> Pins;
+					FAnimNextSchemaAction_AddComment CommentAction;
+					UEdGraphNode* EdNode = CommentAction.PerformAction(EdGraph, Pins, GraphLocation);
+					Controller->CloseUndoBracket();
+				}
+			}
+		});
 
 	Workspace::FObjectDocumentArgs GraphDocumentArgs = WorkspaceEditorModule.CreateGraphDocumentArgs(GraphArgs);
 	Workspace::FOnMakeDocumentWidget WorkspaceMakeDocumentWidgetDelegate = GraphDocumentArgs.OnMakeDocumentWidget;
diff --git a/Engine/Plugins/Experimental/Workspace/Source/WorkspaceEditor/Private/SGraphDocument.cpp b/Engine/Plugins/Experimental/Workspace/Source/WorkspaceEditor/Private/SGraphDocument.cpp
--- a/Engine/Plugins/Experimental/Workspace/Source/WorkspaceEditor/Private/SGraphDocument.cpp
+++ b/Engine/Plugins/Experimental/Workspace/Source/WorkspaceEditor/Private/SGraphDocument.cpp
@@ -31,6 +31,7 @@ void SGraphDocument::Construct(const FArguments& InArgs, TSharedRef<FWorkspaceEd
 	OnDuplicateSelectedNodes = InArgs._OnDuplicateSelectedNodes;
 	OnCanOpenInNewTab = InArgs._OnCanOpenInNewTab;
 	OnOpenInNewTab = InArgs._OnOpenInNewTab;
+	OnCreateComment = InArgs._OnCreateComment;
 
 	SGraphEditor::FGraphEditorEvents Events;
 	Events.OnCreateActionMenuAtLocation = SGraphEditor::FOnCreateActionMenuAtLocation::CreateLambda([this, OnCreateActionMenu = InArgs._OnCreateActionMenu](UEdGraph* InGraph, const FVector2f& InNodePosition, const TArray<UEdGraphPin*>& InDraggedPins, bool bInAutoExpand, SGraphEditor::FActionMenuClosed InOnMenuClosed)
@@ -104,6 +105,9 @@ void SGraphDocument::BindCommands()
 		FExecuteAction::CreateSP(this, &SGraphDocument::OpenInNewTab),
 		FCanExecuteAction::CreateSP(this, &SGraphDocument::CanOpenInNewTab));
 
+	CommandList->MapAction(FGraphEditorCommands::Get().CreateComment,
+		FExecuteAction::CreateSP(this, &SGraphDocument::CreateComment));
+
 }
 
 bool SGraphDocument::CanDeleteSelectedNodes() const
@@ -206,6 +210,11 @@ void SGraphDocument::OpenInNewTab()
 	OnOpenInNewTab.ExecuteIfBound(FWorkspaceEditorContext(HostingAppPtr.Pin().ToSharedRef(), Document), SelectedNodes);
 }
 
+void SGraphDocument::CreateComment()
+{
+	OnCreateComment.ExecuteIfBound(FWorkspaceEditorContext(HostingAppPtr.Pin().ToSharedRef(), Document), GraphEditor->GetPasteLocation2f());
+}
+
 bool SGraphDocument::IsEditable(UEdGraph* InGraph) const
 {
 	return InGraph && HostingAppPtr.Pin()->InEditingMode() && InGraph->bEditable;
diff --git a/Engine/Plugins/Experimental/Workspace/Source/WorkspaceEditor/Private/SGraphDocument.h b/Engine/Plugins/Experimental/Workspace/Source/WorkspaceEditor/Private/SGraphDocument.h
--- a/Engine/Plugins/Experimental/Workspace/Source/WorkspaceEditor/Private/SGraphDocument.h
+++ b/Engine/Plugins/Experimental/Workspace/Source/WorkspaceEditor/Private/SGraphDocument.h
@@ -49,6 +49,8 @@ class SGraphDocument : public SCompoundWidget
 	SLATE_EVENT(FOnCanPerformActionOnSelectedNodes, OnCanOpenInNewTab)
 	SLATE_EVENT(FOnPerformActionOnSelectedNodes, OnOpenInNewTab)
 
+	SLATE_EVENT(FOnCreateComment, OnCreateComment)
+
 	SLATE_END_ARGS()
 
 	void Construct(const FArguments& InArgs, TSharedRef<FWorkspaceEditor> InHostingApp, const FWorkspaceDocument& InDocument);
@@ -77,6 +79,8 @@ class SGraphDocument : public SCompoundWidget
 	bool CanOpenInNewTab() const;
 	void OpenInNewTab();
 
+	void CreateComment();
+
 	bool IsEditable(UEdGraph* InGraph) const;
 	void OnResetSelection();
 
@@ -117,6 +121,8 @@ class SGraphDocument : public SCompoundWidget
 	FOnCanPerformActionOnSelectedNodes OnCanOpenInNewTab;
 	FOnPerformActionOnSelectedNodes OnOpenInNewTab;
 
+	FOnCreateComment OnCreateComment;
+
 	friend class FWorkspaceEditorModule;
 };
 
diff --git a/Engine/Plugins/Experimental/Workspace/Source/WorkspaceEditor/Private/WorkspaceEditorModule.cpp b/Engine/Plugins/Experimental/Workspace/Source/WorkspaceEditor/Private/WorkspaceEditorModule.cpp
--- a/Engine/Plugins/Experimental/Workspace/Source/WorkspaceEditor/Private/WorkspaceEditorModule.cpp
+++ b/Engine/Plugins/Experimental/Workspace/Source/WorkspaceEditor/Private/WorkspaceEditorModule.cpp
@@ -119,7 +119,8 @@ FObjectDocumentArgs FWorkspaceEditorModule::CreateGraphDocumentArgs(const FGraph
 			.OnNavigateHistoryBack_Lambda([WeakWorkspaceEditor](){ if(const TSharedPtr<FWorkspaceEditor> SharedWorkspaceEditor = WeakWorkspaceEditor.Pin()) SharedWorkspaceEditor->NavigateBack(); })
 			.OnNodeDoubleClicked(InArgs.OnNodeDoubleClicked)
 			.OnCanOpenInNewTab(InArgs.OnCanOpenInNewTab)
-			.OnOpenInNewTab(InArgs.OnOpenInNewTab);
+			.OnOpenInNewTab(InArgs.OnOpenInNewTab)
+			.OnCreateComment(InArgs.OnCreateComment);
 
 		// Let listeners know that the document widget has been created 
 		InArgs.OnGraphDocumentCreated.ExecuteIfBound(InContext, Widget);
diff --git a/Engine/Plugins/Experimental/Workspace/Source/WorkspaceEditor/Public/IWorkspaceEditorModule.h b/Engine/Plugins/Experimental/Workspace/Source/WorkspaceEditor/Public/IWorkspaceEditorModule.h
--- a/Engine/Plugins/Experimental/Workspace/Source/WorkspaceEditor/Public/IWorkspaceEditorModule.h
+++ b/Engine/Plugins/Experimental/Workspace/Source/WorkspaceEditor/Public/IWorkspaceEditorModule.h
@@ -189,6 +189,8 @@ using FOnNodeDoubleClicked = TDelegate<void(const FWorkspaceEditorContext&, cons
 
 using FOnGraphDocumentCreated = TDelegate<void(const FWorkspaceEditorContext&, TSharedPtr<SWidget> InGraphWidget)>;
 
+using FOnCreateComment = TDelegate<void(const FWorkspaceEditorContext&, const FVector2f&)>;
+
 // Arguments used to make document widgets for graphs
 struct FGraphDocumentWidgetArgs
 {
@@ -225,6 +227,8 @@ struct FGraphDocumentWidgetArgs
 
 	FOnCanPerformActionOnSelectedNodes OnCanOpenInNewTab;
 	FOnPerformActionOnSelectedNodes OnOpenInNewTab;
+
+	FOnCreateComment OnCreateComment;
 };
 
 // Enum describing how to open a workspace
