Shared Variables asset remains after its removing from Workspace
#fix swap CanDelete default impl to return false

diff --git a/Engine/Plugins/Experimental/Workspace/Source/WorkspaceEditor/Public/IWorkspaceOutlinerItemDetails.h b/Engine/Plugins/Experimental/Workspace/Source/WorkspaceEditor/Public/IWorkspaceOutlinerItemDetails.h
--- a/Engine/Plugins/Experimental/Workspace/Source/WorkspaceEditor/Public/IWorkspaceOutlinerItemDetails.h
+++ b/Engine/Plugins/Experimental/Workspace/Source/WorkspaceEditor/Public/IWorkspaceOutlinerItemDetails.h
@@ -36,7 +36,7 @@ public:
     virtual const FSlateBrush* GetItemIcon(const FWorkspaceOutlinerItemExport& Export) const { return nullptr; }
 	virtual FSlateColor GetItemColor(const FWorkspaceOutlinerItemExport& Export) const { return FSlateColor::UseForeground(); }
     virtual bool HandleDoubleClick(const FToolMenuContext& ToolMenuContext) const { return false; }
-	virtual bool CanDelete(const FWorkspaceOutlinerItemExport& Export) const { return true; }
+	virtual bool CanDelete(const FWorkspaceOutlinerItemExport& Export) const { return false; }
 	virtual void Delete(TConstArrayView<FWorkspaceOutlinerItemExport> Exports) const {}
     virtual bool CanRename(const FWorkspaceOutlinerItemExport& Export) const { return false; }
     virtual void Rename(const FWorkspaceOutlinerItemExport& Export, const FText& InName) const {}
@@ -59,4 +59,4 @@ public:
 
 }
 
-#undef LOCTEXT_NAMESPACE // "IWorkspaceOutlinerItemDetails"
\ No newline at end of file
+#undef LOCTEXT_NAMESPACE // "IWorkspaceOutlinerItemDetails"
