Only add workspace outliner treeview items if a valid ItemDetails implementation for its type is available

diff --git a/Engine/Plugins/Experimental/Workspace/Source/WorkspaceEditor/Private/Outliner/WorkspaceOutlinerHierarchy.cpp b/Engine/Plugins/Experimental/Workspace/Source/WorkspaceEditor/Private/Outliner/WorkspaceOutlinerHierarchy.cpp
--- a/Engine/Plugins/Experimental/Workspace/Source/WorkspaceEditor/Private/Outliner/WorkspaceOutlinerHierarchy.cpp
+++ b/Engine/Plugins/Experimental/Workspace/Source/WorkspaceEditor/Private/Outliner/WorkspaceOutlinerHierarchy.cpp
@@ -5,6 +5,7 @@
 #include "ISceneOutlinerMode.h"
 #include "Workspace.h"
 #include "WorkspaceAssetRegistryInfo.h"
+#include "WorkspaceEditorModule.h"
 #include "WorkspaceOutlinerTreeItem.h"
 #include "AssetRegistry/AssetRegistryModule.h"
 #include "AssetRegistry/IAssetRegistry.h"
@@ -33,9 +34,13 @@ namespace UE::Workspace
 			{
 				auto AddItem = [this, &OutItems](const FWorkspaceOutlinerItemExport& Export)
 				{
-					if (FSceneOutlinerTreeItemPtr Item = Mode->CreateItemFor<FWorkspaceOutlinerTreeItem>(FWorkspaceOutlinerTreeItem::FItemData{Export}))
+					// Only add item types with valid item details
+					if (FWorkspaceEditorModule::GetOutlinerItemDetails(MakeOutlinerDetailsId(Export)))
 					{
-						OutItems.Add(Item);
+						if (FSceneOutlinerTreeItemPtr Item = Mode->CreateItemFor<FWorkspaceOutlinerTreeItem>(FWorkspaceOutlinerTreeItem::FItemData{Export}))
+						{
+							OutItems.Add(Item);
+						}
 					}
 				};
 
diff --git a/Engine/Plugins/Experimental/Workspace/Source/WorkspaceEditor/Private/WorkspaceEditorModule.h b/Engine/Plugins/Experimental/Workspace/Source/WorkspaceEditor/Private/WorkspaceEditorModule.h
--- a/Engine/Plugins/Experimental/Workspace/Source/WorkspaceEditor/Private/WorkspaceEditorModule.h
+++ b/Engine/Plugins/Experimental/Workspace/Source/WorkspaceEditor/Private/WorkspaceEditorModule.h
@@ -93,6 +93,7 @@ private:
 	friend class FWorkspaceOutlinerSourceControlColumn;
 	friend class FWorkspaceOutlinerFileStateColumn;
 	friend class FWorkspaceAssetReferenceOutlinerItemDetails;
+	friend class FWorkspaceOutlinerHierarchy;
 	friend struct FWorkspaceDocument;
 	friend SWorkspaceTabWrapper;
 	friend FWorkspaceEditorContext;
