Fix removing workspace entries using outliner

diff --git a/Engine/Plugins/Experimental/Workspace/Source/WorkspaceEditor/Private/Outliner/WorkspaceOutlinerMode.cpp b/Engine/Plugins/Experimental/Workspace/Source/WorkspaceEditor/Private/Outliner/WorkspaceOutlinerMode.cpp
--- a/Engine/Plugins/Experimental/Workspace/Source/WorkspaceEditor/Private/Outliner/WorkspaceOutlinerMode.cpp
+++ b/Engine/Plugins/Experimental/Workspace/Source/WorkspaceEditor/Private/Outliner/WorkspaceOutlinerMode.cpp
@@ -334,9 +334,13 @@ bool FWorkspaceOutlinerMode::CanDeleteItem(const ISceneOutlinerTreeItem& Item) c
 	{
 		if(TreeItem->ItemDetails.IsValid())
 		{
-			return TreeItem->ItemDetails->CanDelete(TreeItem->Export);
+			if ( TreeItem->ItemDetails->CanDelete(TreeItem->Export))
+			{
+				return true;
+			}
 		}
-		else if(TreeItem->Export.GetParentIdentifier() == NAME_None)
+		
+		if(TreeItem->Export.GetParentIdentifier() == NAME_None)
 		{
 			// No details, but allow deletion if we are a root item
 			return true;
@@ -599,4 +603,4 @@ void FWorkspaceOutlinerMode::OnAssetRegistryAssetUpdate(const FAssetData& AssetD
 }
 }	
 
-#undef LOCTEXT_NAMESPACE // "FWorkspaceOutlinerMode"
\ No newline at end of file
+#undef LOCTEXT_NAMESPACE // "FWorkspaceOutlinerMode"
