Workspace Changes Round 1
* Update workspace schema to control if a workspace can be saved and if a workspace can contain multiple assets.
* Move 'Save All' button from the asset toolbar to the workspace outliner tab.
* Prevent creating new UAF workspace assets and prevent opening old UAF workspaces.
* Save/Browse buttons in asset toolbar now relate to the focused document.

diff --git a/Engine/Plugins/Experimental/UAF/UAF/Source/UAFEditor/Internal/Workspace/AnimNextWorkspaceFactory.h b/Engine/Plugins/Experimental/UAF/UAF/Source/UAFEditor/Internal/Workspace/AnimNextWorkspaceFactory.h
--- a/Engine/Plugins/Experimental/UAF/UAF/Source/UAFEditor/Internal/Workspace/AnimNextWorkspaceFactory.h
+++ b/Engine/Plugins/Experimental/UAF/UAF/Source/UAFEditor/Internal/Workspace/AnimNextWorkspaceFactory.h
@@ -18,4 +18,5 @@ class UAnimNextWorkspaceFactory : public UWorkspaceFactory
 	virtual FText GetDisplayName() const override;
 	virtual uint32 GetMenuCategories() const override;
 	virtual const TArray<FText>& GetMenuCategorySubMenus() const override;
+	virtual bool ShouldShowInNewMenu() const override;
 };
\ No newline at end of file
diff --git a/Engine/Plugins/Experimental/UAF/UAF/Source/UAFEditor/Private/Workspace/AnimNextWorkspaceFactory.cpp b/Engine/Plugins/Experimental/UAF/UAF/Source/UAFEditor/Private/Workspace/AnimNextWorkspaceFactory.cpp
--- a/Engine/Plugins/Experimental/UAF/UAF/Source/UAFEditor/Private/Workspace/AnimNextWorkspaceFactory.cpp
+++ b/Engine/Plugins/Experimental/UAF/UAF/Source/UAFEditor/Private/Workspace/AnimNextWorkspaceFactory.cpp
@@ -34,4 +34,9 @@ const TArray<FText>& UAnimNextWorkspaceFactory::GetMenuCategorySubMenus() const
 	return Categories;
 }
 
+bool UAnimNextWorkspaceFactory::ShouldShowInNewMenu() const
+{
+	return false;
+}
+
 #undef LOCTEXT_NAMESPACE
diff --git a/Engine/Plugins/Experimental/UAF/UAF/Source/UAFEditor/Private/Workspace/AnimNextWorkspaceSchema.h b/Engine/Plugins/Experimental/UAF/UAF/Source/UAFEditor/Private/Workspace/AnimNextWorkspaceSchema.h
--- a/Engine/Plugins/Experimental/UAF/UAF/Source/UAFEditor/Private/Workspace/AnimNextWorkspaceSchema.h
+++ b/Engine/Plugins/Experimental/UAF/UAF/Source/UAFEditor/Private/Workspace/AnimNextWorkspaceSchema.h
@@ -20,4 +20,7 @@ class UAnimNextWorkspaceSchema : public UWorkspaceSchema
 	virtual void OnLoadWorkspaceState(TSharedRef<UE::Workspace::IWorkspaceEditor> InWorkspaceEditor, const FInstancedStruct& InWorkspaceState) const override;
 	virtual bool SupportsViewport() const { return true; };
 	virtual TObjectPtr<UWorkspaceViewportSceneDescription> CreateSceneDescription() const { return NewObject<UUAFViewportSceneDescription>(); }
+	virtual bool SupportsMultipleAssets() const { return false; }
+	virtual bool HasPersistentRootAsset() const { return true; }
+	virtual bool CanSaveWorkspace() const { return false; }
 };
\ No newline at end of file
diff --git a/Engine/Plugins/Experimental/UAF/UAFTests/Source/UAFTests/Private/UAFAssetCreationTests.cpp b/Engine/Plugins/Experimental/UAF/UAFTests/Source/UAFTests/Private/UAFAssetCreationTests.cpp
--- a/Engine/Plugins/Experimental/UAF/UAFTests/Source/UAFTests/Private/UAFAssetCreationTests.cpp
+++ b/Engine/Plugins/Experimental/UAF/UAFTests/Source/UAFTests/Private/UAFAssetCreationTests.cpp
@@ -29,22 +29,6 @@ TEST_CLASS(FAssetCreationTests, "Animation.UAF.Functional")
 		AssetEditorSubsystem = GEditor->GetEditorSubsystem<UAssetEditorSubsystem>();
 		ASSERT_THAT(IsNotNull(AssetEditorSubsystem));
 	}
-
-	TEST_METHOD(Verify_Workspace_Asset_Creation)
-	{		
-		TestCommandBuilder
-			.Do(TEXT("Create AnimNext Workspace Asset"), [&]()
-			{
-				const TSubclassOf<UAnimNextWorkspaceFactory> WorkSpaceFactoryClass = UAnimNextWorkspaceFactory::StaticClass();
-				UObject* FactoryObject = UAFTestsUtilities::CreateFactoryObject(NewObject<UAnimNextWorkspaceFactory>(), UWorkspace::StaticClass(), "NewAnimNextWorkspace");
-				AnimNextAsset = CastChecked<UWorkspace>(FactoryObject);
-			})
-			.Then(TEXT("Assert AnimNext Workspace Asset Created and Open in Editor"), [&]()
-			{
-				ASSERT_THAT(IsNotNull(AnimNextAsset));
-				ASSERT_THAT(IsTrue(AssetEditorSubsystem->OpenEditorForAsset(AnimNextAsset)));
-			});
-	}
 	
 	TEST_METHOD(Verify_StateTree_Asset_Creation)
 	{
diff --git a/Engine/Plugins/Experimental/UAF/UAFTests/Source/UAFTests/Private/UAFModuleTests.cpp b/Engine/Plugins/Experimental/UAF/UAFTests/Source/UAFTests/Private/UAFModuleTests.cpp
--- a/Engine/Plugins/Experimental/UAF/UAFTests/Source/UAFTests/Private/UAFModuleTests.cpp
+++ b/Engine/Plugins/Experimental/UAF/UAFTests/Source/UAFTests/Private/UAFModuleTests.cpp
@@ -62,34 +62,6 @@ TEST_CLASS(FAnimNextModuleTests, "Animation.UAF.Functional")
 			AnimNextModule = nullptr;
 		}
 	}
-	
-	// QMetry Test Case:  UE-TC-10714
-	TEST_METHOD(Save_AnimNext_Module_Without_Workspace)
-	{
-		TestCommandBuilder
-			.Do(TEXT("Create AnimNext Workspaxce Asset"), [&]()
-			{
-				CreateAnimNextWorkspace();
-			})
-			.Then(TEXT("Create AnimNext Module Asset"), [&]()
-			{
-				CreateAnimNextModule();
-			})
-			.Then(TEXT("Assert that UAF Assets were created"), [&]()
-			{
-				ASSERT_THAT(IsNotNull(AnimNextWorkspace));
-				ASSERT_THAT(IsTrue(AssetEditorSubsystem->OpenEditorForAsset(AnimNextWorkspace)));
-				ASSERT_THAT(IsNotNull(AnimNextModule));
-				ASSERT_THAT(IsTrue(AssetEditorSubsystem->OpenEditorForAsset(AnimNextModule)));
-			})
-			.Then(TEXT("Assert that UAF Assets were saved"), [&]()
-			{
-				FSavePackageArgs SaveArgs;
-				SaveArgs.TopLevelFlags = RF_Standalone;				
-				ASSERT_THAT(IsTrue(UPackage::SavePackage(AnimNextWorkspacePackage, nullptr, *(AnimNextWorkspacePackageName.ToString()), SaveArgs)));
-				ASSERT_THAT(IsTrue(UPackage::SavePackage(AnimNextModulePackage, nullptr, *(AnimNextModulePackageName.ToString()), SaveArgs)));
-			});
-	}
 		
 	// QMetry Test Case:  UE-TC-18778
 	TEST_METHOD(Undo_Redo_In_AnimNext_Module)
@@ -319,16 +291,6 @@ protected:
 		AnimNextModule = CastChecked<UAnimNextModule>(FactoryObject);
 	}
 	
-	void CreateAnimNextWorkspace()
-	{
-		UFactory* Factory = NewObject<UAnimNextWorkspaceFactory>();
-		AnimNextWorkspacePackage = CreateUniquePackage("NewAnimNextWorkspace");
-		AnimNextWorkspacePackageName = *FPaths::GetBaseFilename(AnimNextWorkspacePackage->GetName());		
-		UObject* FactoryObject = Factory->FactoryCreateNew(UWorkspace::StaticClass(), AnimNextWorkspacePackage, AnimNextWorkspacePackageName, RF_Public | RF_Standalone, NULL, GWarn);
-		FAssetRegistryModule::AssetCreated(FactoryObject);
-		AnimNextWorkspace = CastChecked<UWorkspace>(FactoryObject);
-	}
-	
 	UPackage* CreateUniquePackage(FString InAssetName)
 	{
 		// Create a unique package name and path
diff --git a/Engine/Plugins/Experimental/UAF/UAFTests/Source/UAFTests/Private/UAFWorkspaceTests.cpp b/Engine/Plugins/Experimental/UAF/UAFTests/Source/UAFTests/Private/UAFWorkspaceTests.cpp
--- a/Engine/Plugins/Experimental/UAF/UAFTests/Source/UAFTests/Private/UAFWorkspaceTests.cpp
+++ b/Engine/Plugins/Experimental/UAF/UAFTests/Source/UAFTests/Private/UAFWorkspaceTests.cpp
@@ -6,6 +6,7 @@
 #include "UAFTestsUtilities.h"
 #include "CQTest.h"
 #include "Editor.h"
+#include "IWorkspaceEditor.h"
 #include "Subsystems/AssetEditorSubsystem.h"
 #include "WorkspaceEditor/Private/Workspace.h"
 #include "Workspace/AnimNextWorkspaceFactory.h"
@@ -16,9 +17,8 @@
 
 TEST_CLASS(FUAFWorkspaceTests, "Animation.UAF.Functional")
 {
-	UWorkspace* WorkspaceAsset = nullptr;
-	UWorkspace* WorkspaceSecondaryAsset = nullptr;
 	UAnimNextAnimationGraph* AnimationGraphAsset = nullptr;
+	UAnimNextAnimationGraph* AnimationGraphAsset2 = nullptr;
 	UAnimNextModule* ModuleAsset = nullptr;
 	UAssetEditorSubsystem* AssetEditorSubsystem = nullptr;
 
@@ -29,7 +29,7 @@ TEST_CLASS(FUAFWorkspaceTests, "Animation.UAF.Functional")
 	}
 
 	// QMetry Test Case:  UE-TC-10700
-	TEST_METHOD(Verify_Add_Animation_Graph_To_Workspace)
+	TEST_METHOD(Verify_Open_Animation_Graph_As_Root_In_Workspace)
 	{
 		TestCommandBuilder
 			.Do(TEXT("Create Animation Graph Asset"), [&]()
@@ -39,51 +39,63 @@ TEST_CLASS(FUAFWorkspaceTests, "Animation.UAF.Functional")
 				AnimationGraphAsset = CastChecked<UAnimNextAnimationGraph>(FactoryObject);
 				ASSERT_THAT(IsNotNull(AnimationGraphAsset));
 			})
-			.Then(TEXT("Create UAF Workspace Asset"), [&]()
+			.Then(TEXT("Open Animation Graph in UAF Workspace"), [&]()
 			{
-				const TSubclassOf<UAnimNextWorkspaceFactory> WorkSpaceFactoryClass = UAnimNextWorkspaceFactory::StaticClass();
-				UObject* FactoryObject = UAFTestsUtilities::CreateFactoryObject(NewObject<UAnimNextWorkspaceFactory>(), UWorkspace::StaticClass(), "NewAnimNextWorkspace");
-				WorkspaceAsset = CastChecked<UWorkspace>(FactoryObject);
-				ASSERT_THAT(IsNotNull(WorkspaceAsset));
+				ASSERT_THAT(IsTrue(AssetEditorSubsystem->OpenEditorForAsset(AnimationGraphAsset)));
+			});
+	}
+
+	TEST_METHOD(Verify_Cannot_Remove_Workspace_Root_Asset)
+	{
+		TestCommandBuilder
+			.Do(TEXT("Create Animation Graph Asset"), [&]()
+			{
+				const TSubclassOf<UAnimNextAnimationGraphFactory> AnimGraphFactoryClass = UAnimNextAnimationGraphFactory::StaticClass();
+				UObject* FactoryObject = UAFTestsUtilities::CreateFactoryObject(NewObject<UAnimNextAnimationGraphFactory>(), UAnimNextAnimationGraph::StaticClass(), "NewAnimNextAnimationGraph");
+				AnimationGraphAsset = CastChecked<UAnimNextAnimationGraph>(FactoryObject);
+				ASSERT_THAT(IsNotNull(AnimationGraphAsset));
 			})
-			.Then(TEXT("Add Animation Graph to Workspace and Open in Editor"), [&]()
+			.Then(TEXT("Open Animation Graph in UAF Workspace"), [&]()
 			{
-				ASSERT_THAT(IsTrue(WorkspaceAsset->AddAsset(AnimationGraphAsset)));
-				ASSERT_THAT(IsTrue(AssetEditorSubsystem->OpenEditorForAsset(WorkspaceAsset)));
+				ASSERT_THAT(IsTrue(AssetEditorSubsystem->OpenEditorForAsset(AnimationGraphAsset)));
+			})
+			.Then(TEXT("Try remove Animation Graph from Workspace"), [&]()
+			{
+				IAssetEditorInstance* AssetEditor = AssetEditorSubsystem->FindEditorForAsset(AnimationGraphAsset, false);
+				UE::Workspace::IWorkspaceEditor* WorkspaceEditor = static_cast<UE::Workspace::IWorkspaceEditor*>(AssetEditor);
+				UWorkspace* Workspace = CastChecked<UWorkspace>(WorkspaceEditor->GetWorkspaceAsset());
+
+				TestRunner->AddExpectedError("Cannot remove the root-level asset.", EAutomationExpectedErrorFlags::Contains, -1);
+				ASSERT_THAT(IsFalse(Workspace->RemoveAsset(AnimationGraphAsset)));
 			});
 	}
-	
-	// QMetry Test Case:  UE-TC-10716
-	TEST_METHOD(Verify_Add_Module_To_Multiple_Workspaces)
+
+	TEST_METHOD(Verify_Cannot_Add_Multiple_Assets_To_Workspace)
 	{
 		TestCommandBuilder
-			.Do(TEXT("Create UAF Module Asset"), [&]()
+			.Do(TEXT("Create Animation Graph Assets"), [&]()
 			{
-				const TSubclassOf<UAnimNextModuleFactory> ModuleFactoryClass = UAnimNextModuleFactory::StaticClass();
-				UObject* FactoryObject = UAFTestsUtilities::CreateFactoryObject(NewObject<UAnimNextModuleFactory>(), UAnimNextModule::StaticClass(), "NewAnimNextModule");				
-				ModuleAsset = CastChecked<UAnimNextModule>(FactoryObject);
-				ASSERT_THAT(IsNotNull(ModuleAsset));
+				const TSubclassOf<UAnimNextAnimationGraphFactory> AnimGraphFactoryClass = UAnimNextAnimationGraphFactory::StaticClass();
+				UObject* FactoryObject = UAFTestsUtilities::CreateFactoryObject(NewObject<UAnimNextAnimationGraphFactory>(), UAnimNextAnimationGraph::StaticClass(), "NewAnimNextAnimationGraph");
+				AnimationGraphAsset = CastChecked<UAnimNextAnimationGraph>(FactoryObject);
+				ASSERT_THAT(IsNotNull(AnimationGraphAsset));
+
+				UObject* FactoryObject2 = UAFTestsUtilities::CreateFactoryObject(NewObject<UAnimNextAnimationGraphFactory>(), UAnimNextAnimationGraph::StaticClass(), "NewAnimNextAnimationGraph");
+				AnimationGraphAsset2 = CastChecked<UAnimNextAnimationGraph>(FactoryObject2);
+				ASSERT_THAT(IsNotNull(AnimationGraphAsset2));
 			})
-			.Then(TEXT("Create UAF Workspace Asset"), [&]()
+			.Then(TEXT("Open Animation Graph in UAF Workspace"), [&]()
 			{
-				const TSubclassOf<UAnimNextWorkspaceFactory> WorkSpaceFactoryClass = UAnimNextWorkspaceFactory::StaticClass();
-				UObject* FactoryObject = UAFTestsUtilities::CreateFactoryObject(NewObject<UAnimNextWorkspaceFactory>(), UWorkspace::StaticClass(), "NewAnimNextWorkspace");
-				WorkspaceAsset = CastChecked<UWorkspace>(FactoryObject);
-				ASSERT_THAT(IsNotNull(WorkspaceAsset));
+				ASSERT_THAT(IsTrue(AssetEditorSubsystem->OpenEditorForAsset(AnimationGraphAsset)));
 			})
-			.Then(TEXT("Create Another UAF Workspace Asset"), [&]()
+			.Then(TEXT("Try add second Animation Graph to Workspace"), [&]()
 			{
-				const TSubclassOf<UAnimNextWorkspaceFactory> WorkSpaceFactoryClass = UAnimNextWorkspaceFactory::StaticClass();
-				UObject* FactoryObject = UAFTestsUtilities::CreateFactoryObject(NewObject<UAnimNextWorkspaceFactory>(), UWorkspace::StaticClass(), "NewAnimNextWorkspace1");
-				WorkspaceSecondaryAsset = CastChecked<UWorkspace>(FactoryObject);
-				ASSERT_THAT(IsNotNull(WorkspaceSecondaryAsset));
-			})		
-			.Then(TEXT("Add Module to Multiple Workspaces and Open in Editor"), [&]()
-			{
-				ASSERT_THAT(IsTrue(WorkspaceAsset->AddAsset(ModuleAsset)));
-				ASSERT_THAT(IsTrue(WorkspaceSecondaryAsset->AddAsset(ModuleAsset)));
-				ASSERT_THAT(IsTrue(AssetEditorSubsystem->OpenEditorForAsset(WorkspaceAsset)));
-				ASSERT_THAT(IsTrue(AssetEditorSubsystem->OpenEditorForAsset(WorkspaceSecondaryAsset)));
+				IAssetEditorInstance* AssetEditor = AssetEditorSubsystem->FindEditorForAsset(AnimationGraphAsset, false);
+				UE::Workspace::IWorkspaceEditor* WorkspaceEditor = static_cast<UE::Workspace::IWorkspaceEditor*>(AssetEditor);
+				UWorkspace* Workspace = CastChecked<UWorkspace>(WorkspaceEditor->GetWorkspaceAsset());
+
+				TestRunner->AddExpectedError("Schema does not support adding multiple assets.", EAutomationExpectedErrorFlags::Contains, -1);
+				ASSERT_THAT(IsFalse(Workspace->AddAsset(AnimationGraphAsset2)));
 			});
 	}
 };
diff --git a/Engine/Plugins/Experimental/Workspace/Source/WorkspaceEditor/Private/DefaultWorkspaceSchema.h b/Engine/Plugins/Experimental/Workspace/Source/WorkspaceEditor/Private/DefaultWorkspaceSchema.h
--- a/Engine/Plugins/Experimental/Workspace/Source/WorkspaceEditor/Private/DefaultWorkspaceSchema.h
+++ b/Engine/Plugins/Experimental/Workspace/Source/WorkspaceEditor/Private/DefaultWorkspaceSchema.h
@@ -14,4 +14,6 @@ class UDefaultWorkspaceSchema : public UWorkspaceSchema
 	// UWorkspaceSchema interface
 	virtual FText GetDisplayName() const { return FText::GetEmpty(); }
 	virtual TConstArrayView<FTopLevelAssetPath> GetSupportedAssetClassPaths() const override { return TConstArrayView<FTopLevelAssetPath>(); }
+	virtual bool SupportsMultipleAssets() const override { return true; }
+	virtual bool CanSaveWorkspace() const override { return true; }
 };
\ No newline at end of file
diff --git a/Engine/Plugins/Experimental/Workspace/Source/WorkspaceEditor/Private/Outliner/WorkspaceOutlinerMode.cpp b/Engine/Plugins/Experimental/Workspace/Source/WorkspaceEditor/Private/Outliner/WorkspaceOutlinerMode.cpp
--- a/Engine/Plugins/Experimental/Workspace/Source/WorkspaceEditor/Private/Outliner/WorkspaceOutlinerMode.cpp
+++ b/Engine/Plugins/Experimental/Workspace/Source/WorkspaceEditor/Private/Outliner/WorkspaceOutlinerMode.cpp
@@ -18,6 +18,7 @@
 #include "WorkspaceEditorCommands.h"
 #include "Framework/Commands/GenericCommands.h"
 #include "Algo/ForEach.h"
+#include "Framework/Application/SlateApplication.h"
 
 #define LOCTEXT_NAMESPACE "FWorkspaceOutlinerMode"
 
@@ -464,7 +465,12 @@ void FWorkspaceOutlinerMode::OpenItems(TArrayView<const FSceneOutlinerTreeItemPt
 			
 			if(!bHandled)
 			{
-				if (const TSharedPtr<UE::Workspace::IWorkspaceEditor> SharedWorkspaceEditor = StaticCastSharedPtr<UE::Workspace::IWorkspaceEditor>(WeakWorkspaceEditor.Pin()))
+				if (FSlateApplication::Get().GetModifierKeys().IsShiftDown())
+				{
+					UAssetEditorSubsystem* AssetEditorSubsystem = GEditor->GetEditorSubsystem<UAssetEditorSubsystem>();
+					AssetEditorSubsystem->OpenEditorForAsset(TreeItem->Export.GetFirstAssetPath());
+				}
+				else if (const TSharedPtr<UE::Workspace::IWorkspaceEditor> SharedWorkspaceEditor = StaticCastSharedPtr<UE::Workspace::IWorkspaceEditor>(WeakWorkspaceEditor.Pin()))
 				{
 					SharedWorkspaceEditor->OpenExports({TreeItem->Export});
 				}
diff --git a/Engine/Plugins/Experimental/Workspace/Source/WorkspaceEditor/Private/SWorkspaceView.cpp b/Engine/Plugins/Experimental/Workspace/Source/WorkspaceEditor/Private/SWorkspaceView.cpp
--- a/Engine/Plugins/Experimental/Workspace/Source/WorkspaceEditor/Private/SWorkspaceView.cpp
+++ b/Engine/Plugins/Experimental/Workspace/Source/WorkspaceEditor/Private/SWorkspaceView.cpp
@@ -11,6 +11,7 @@
 #include "SPositiveActionButton.h"
 #include "SSceneOutliner.h"
 #include "WorkspaceAssetRegistryInfo.h"
+#include "WorkspaceSchema.h"
 #include "Outliner/WorkspaceOutlinerColumns.h"
 #include "Outliner/WorkspaceOutlinerHierarchy.h"
 #include "Outliner/WorkspaceOutlinerMode.h"
@@ -24,16 +25,29 @@ namespace UE::Workspace
 class SWorkspaceOutliner : public SSceneOutliner
 {
 public:
+	SLATE_BEGIN_ARGS(SWorkspaceOutliner) {}
+		SLATE_ATTRIBUTE(bool, CanSaveAll)
+		SLATE_EVENT(FOnClicked, OnSaveAll)
+	SLATE_END_ARGS()
+	
 	void Construct(const FArguments& InArgs, const FSceneOutlinerInitializationOptions& InitOptions, UWorkspace* InWorkspace)
 	{
 		AddAssetButton = SNew(SPositiveActionButton)
 				.OnGetMenuContent(this, &SWorkspaceOutliner::GetAddAssetPicker)			 
 				.Icon(FAppStyle::Get().GetBrush("Icons.Plus"))
 				.Text(LOCTEXT("AddAssetButton", "Add"));
+				
+		SaveAllButton = SNew(SActionButton)
+			.IsEnabled(InArgs._CanSaveAll)
+			.OnClicked(InArgs._OnSaveAll)
+			.Text(LOCTEXT("SaveAll", "Save All"))
+			.Icon(FAppStyle::GetBrush("MainFrame.SaveAll"));
 
 		WeakWorkspace = InWorkspace;
 		
-		SSceneOutliner::Construct(InArgs, InitOptions);
+		SSceneOutliner::FArguments OutlinerArgs;
+		
+		SSceneOutliner::Construct(OutlinerArgs, InitOptions);
 	}
 
 	SWorkspaceOutliner() = default;
@@ -41,12 +55,26 @@ public:
 
 	virtual void CustomAddToToolbar(TSharedPtr<class SHorizontalBox> Toolbar) override
 	{
+		if (UWorkspace* Workspace = WeakWorkspace.Get())
+		{
+			if (Workspace->GetSchema()->SupportsMultipleAssets())
+			{
+				Toolbar->AddSlot()
+				.VAlign(VAlign_Center)
+				.AutoWidth()
+				.Padding(4.f, 0.f, 0.f, 0.f)
+				[
+					AddAssetButton.ToSharedRef()
+				];
+			}
+		}
+		
 		Toolbar->AddSlot()
 		.VAlign(VAlign_Center)
 		.AutoWidth()
 		.Padding(4.f, 0.f, 0.f, 0.f)
 		[
-			AddAssetButton.ToSharedRef()
+			SaveAllButton.ToSharedRef()
 		];
 	}
 
@@ -85,10 +113,11 @@ public:
 		return SNullWidget::NullWidget;	
 	
 	}
-private:
 	
+private:
 	TWeakObjectPtr<UWorkspace> WeakWorkspace;
 	TSharedPtr<SPositiveActionButton> AddAssetButton;
+	TSharedPtr<SActionButton> SaveAllButton;
 };
 
 void SWorkspaceView::Construct(const FArguments& InArgs, UWorkspace* InWorkspace, TSharedRef<UE::Workspace::IWorkspaceEditor> InWorkspaceEditor)
@@ -108,7 +137,9 @@ void SWorkspaceView::Construct(const FArguments& InArgs, UWorkspace* InWorkspace
 			return OutlinerMode;
 		});
 	}
-	SceneWorkspaceOutliner = SNew(SWorkspaceOutliner, InitOptions, Workspace);
+	SceneWorkspaceOutliner = SNew(SWorkspaceOutliner, InitOptions, Workspace)
+		.CanSaveAll(InArgs._CanSaveAll)
+		.OnSaveAll(InArgs._OnSaveAll);
 
 	TWeakObjectPtr<UWorkspace> WeakWorkspace = Workspace;
 	ChildSlot
@@ -122,6 +153,12 @@ void SWorkspaceView::Construct(const FArguments& InArgs, UWorkspace* InWorkspace
 		})
 		.OnAreAssetsAcceptableForDropWithReason_Lambda([this](TArrayView<FAssetData> InAssets, FText& OutText)
 		{
+			if (!Workspace->GetSchema()->SupportsMultipleAssets())
+			{
+				OutText = LOCTEXT("MultipleAssetsUnsupportedInWorkspace", "Multiple assets are not supported by this workspace");
+				return false;
+			}
+			
 			for(const FAssetData& Asset : InAssets)
 			{
 				if(Workspace->IsAssetSupported(Asset))
diff --git a/Engine/Plugins/Experimental/Workspace/Source/WorkspaceEditor/Private/SWorkspaceView.h b/Engine/Plugins/Experimental/Workspace/Source/WorkspaceEditor/Private/SWorkspaceView.h
--- a/Engine/Plugins/Experimental/Workspace/Source/WorkspaceEditor/Private/SWorkspaceView.h
+++ b/Engine/Plugins/Experimental/Workspace/Source/WorkspaceEditor/Private/SWorkspaceView.h
@@ -18,6 +18,8 @@ class FWorkspaceOutlinerMode;
 class SWorkspaceView : public SCompoundWidget
 {
 	SLATE_BEGIN_ARGS(SWorkspaceView) {}
+		SLATE_ATTRIBUTE(bool, CanSaveAll)
+		SLATE_EVENT(FOnClicked, OnSaveAll)
 	SLATE_END_ARGS()
 
 	void Construct(const FArguments& InArgs, UWorkspace* InWorkspace, TSharedRef<UE::Workspace::IWorkspaceEditor> InWorkspaceEditor);
diff --git a/Engine/Plugins/Experimental/Workspace/Source/WorkspaceEditor/Private/Workspace.cpp b/Engine/Plugins/Experimental/Workspace/Source/WorkspaceEditor/Private/Workspace.cpp
--- a/Engine/Plugins/Experimental/Workspace/Source/WorkspaceEditor/Private/Workspace.cpp
+++ b/Engine/Plugins/Experimental/Workspace/Source/WorkspaceEditor/Private/Workspace.cpp
@@ -17,6 +17,12 @@ REGISTER_ASSETDEPENDENCY_GATHERER(FExternalAssetDependencyGatherer, UWorkspace);
 
 bool UWorkspace::AddAsset(const FAssetData& InAsset, bool bSetupUndoRedo, bool bPrintPythonCommand)
 {
+	if (!GetSchema()->SupportsMultipleAssets() && !AssetEntries.IsEmpty())
+	{
+		ReportError(TEXT("UWorkspace::AddAssets: Schema does not support adding multiple assets."));
+		return false;
+	}
+	
 	if(!IsAssetSupported(InAsset))
 	{
 		ReportError(TEXT("UWorkspace::AddAsset: Unsupported asset supplied."));
@@ -120,6 +126,14 @@ bool UWorkspace::AddAssets(const TArray<UObject*>& InAssets, bool bSetupUndoRedo
 
 bool UWorkspace::RemoveAsset(const FAssetData& InAsset, bool bSetupUndoRedo, bool bPrintPythonCommand)
 {
+	const bool bIsRootLevelAsset = GetRootAsset() == InAsset.GetAsset();
+	
+	if (bIsRootLevelAsset)
+	{
+		ReportError(TEXT("UWorkspace::RemoveAsset: Cannot remove the root-level asset."));
+		return false;
+	}
+
 	if(!IsAssetSupported(InAsset))
 	{
 		ReportError(TEXT("UWorkspace::RemoveAsset: Unsupported asset supplied."));
@@ -371,3 +385,18 @@ void UWorkspace::PostDuplicate(bool bDuplicateForPIE)
 
 	Guid = FGuid::NewGuid();
 }
+
+UObject* UWorkspace::GetRootAsset() const
+{
+	if (!GetSchema()->HasPersistentRootAsset())
+	{
+		return nullptr;
+	}
+
+	check(!AssetEntries.IsEmpty());
+	UWorkspaceAssetEntry* RootAssetEntry = AssetEntries[0];
+	UObject* RootAsset = RootAssetEntry->Asset.LoadSynchronous();
+	check(RootAsset);
+	
+	return RootAsset;
+}
\ No newline at end of file
diff --git a/Engine/Plugins/Experimental/Workspace/Source/WorkspaceEditor/Private/Workspace.h b/Engine/Plugins/Experimental/Workspace/Source/WorkspaceEditor/Private/Workspace.h
--- a/Engine/Plugins/Experimental/Workspace/Source/WorkspaceEditor/Private/Workspace.h
+++ b/Engine/Plugins/Experimental/Workspace/Source/WorkspaceEditor/Private/Workspace.h
@@ -102,7 +102,7 @@ class UWorkspace : public UObject
 	// Removes an asset from the workspace
 	// @return true if the asset was removed
 	UFUNCTION(BlueprintCallable, Category = "Workspace")
-	bool RemoveAsset(UObject* InAsset, bool bSetupUndoRedo = true, bool bPrintPythonCommand = true);
+	WORKSPACEEDITOR_API bool RemoveAsset(UObject* InAsset, bool bSetupUndoRedo = true, bool bPrintPythonCommand = true);
 	bool RemoveAsset(const FAssetData& InAsset, bool bSetupUndoRedo = true, bool bPrintPythonCommand = true);
 
 	// Removes assets from the workspace
@@ -123,6 +123,11 @@ class UWorkspace : public UObject
 
 	void BroadcastModified();
 
+	// Returns the "root asset" that this workspace was created for.
+	// Returns nullptr if the schema does not support a root asset.
+	// Root assets cannot be removed from the workspace.
+	UObject* GetRootAsset() const;
+
 	// Returns all contained Asset objects
 	void GetAssets(TArray<TObjectPtr<UObject>>& OutAssets) const;
 	// Returns all contained AssetData entries
diff --git a/Engine/Plugins/Experimental/Workspace/Source/WorkspaceEditor/Private/WorkspaceAssetDefinition.cpp b/Engine/Plugins/Experimental/Workspace/Source/WorkspaceEditor/Private/WorkspaceAssetDefinition.cpp
--- a/Engine/Plugins/Experimental/Workspace/Source/WorkspaceEditor/Private/WorkspaceAssetDefinition.cpp
+++ b/Engine/Plugins/Experimental/Workspace/Source/WorkspaceEditor/Private/WorkspaceAssetDefinition.cpp
@@ -4,6 +4,8 @@
 #include "WorkspaceEditor.h"
 #include "WorkspaceSchema.h"
 #include "WorkspaceAssetEditor.h"
+#include "Framework/Notifications/NotificationManager.h"
+#include "Widgets/Notifications/SNotificationList.h"
 
 #include UE_INLINE_GENERATED_CPP_BY_NAME(WorkspaceAssetDefinition)
 
@@ -36,6 +38,20 @@ EAssetCommandResult UAssetDefinition_Workspace::OpenAssets(const FAssetOpenArgs&
 {
 	for (UWorkspace* Asset : OpenArgs.LoadObjects<UWorkspace>())
 	{
+		if (!Asset->GetSchema()->CanSaveWorkspace())
+		{
+			// Attempting to open a workspace asset that is no longer savable as an asset
+			FNotificationInfo ErrorNotification(LOCTEXT("CannotOpenWorkspaceAsset", "Workspace assets of this schema can no longer be saved as assets. Please delete this asset."));
+			ErrorNotification.bFireAndForget = true;
+			ErrorNotification.bUseLargeFont = false;
+			ErrorNotification.ExpireDuration = 3.0f;
+			
+            FSlateNotificationManager::Get().AddNotification(ErrorNotification)->SetCompletionState(SNotificationItem::CS_Fail);
+			UE_LOG(LogAnimation, Error, TEXT("Workspace assets of this schema can no longer be saved as assets. Please delete this asset."));
+			
+			return EAssetCommandResult::Handled;
+		}
+	
 		UAssetEditorSubsystem* AssetEditorSubsystem = GEditor->GetEditorSubsystem<UAssetEditorSubsystem>();
 		UWorkspaceAssetEditor* AssetEditor = NewObject<UWorkspaceAssetEditor>(AssetEditorSubsystem, NAME_None, RF_Transient);
 		AssetEditor->SetObjectToEdit(Asset);
diff --git a/Engine/Plugins/Experimental/Workspace/Source/WorkspaceEditor/Private/WorkspaceAssetEditor.h b/Engine/Plugins/Experimental/Workspace/Source/WorkspaceEditor/Private/WorkspaceAssetEditor.h
--- a/Engine/Plugins/Experimental/Workspace/Source/WorkspaceEditor/Private/WorkspaceAssetEditor.h
+++ b/Engine/Plugins/Experimental/Workspace/Source/WorkspaceEditor/Private/WorkspaceAssetEditor.h
@@ -20,6 +20,7 @@ protected:
 	virtual TSharedPtr<FBaseAssetToolkit> CreateToolkit() override;
 
 private:
+	UPROPERTY()
 	TObjectPtr<UWorkspace> ObjectToEdit;
 };
 
diff --git a/Engine/Plugins/Experimental/Workspace/Source/WorkspaceEditor/Private/WorkspaceEditor.cpp b/Engine/Plugins/Experimental/Workspace/Source/WorkspaceEditor/Private/WorkspaceEditor.cpp
--- a/Engine/Plugins/Experimental/Workspace/Source/WorkspaceEditor/Private/WorkspaceEditor.cpp
+++ b/Engine/Plugins/Experimental/Workspace/Source/WorkspaceEditor/Private/WorkspaceEditor.cpp
@@ -31,6 +31,9 @@
 #include "WorkspaceAssetViewportClient.h"
 #include "WorkspaceViewportController.h"
 #include "Viewports.h"
+#include "AssetToolsModule.h"
+#include "IAssetTools.h"
+#include "Styling/SlateIconFinder.h"
 
 #define LOCTEXT_NAMESPACE "WorkspaceEditor"
 
@@ -214,7 +217,16 @@ void FWorkspaceEditor::CreateWidgets()
 	WorkspaceEditorModule.OnExtendTabs().Broadcast(*LayoutExtender.Get(), StaticCastSharedRef<UE::Workspace::IWorkspaceEditor>(AsShared()));
 	StandaloneDefaultLayout->ProcessExtensions(*LayoutExtender.Get());
 
-	WorkspaceView = SNew(SWorkspaceView, Workspace, StaticCastSharedRef<UE::Workspace::IWorkspaceEditor>(AsShared()));
+	WorkspaceView = SNew(SWorkspaceView, Workspace, StaticCastSharedRef<UE::Workspace::IWorkspaceEditor>(AsShared()))
+		.CanSaveAll_Lambda([this]()
+		{
+			return AreAssetEntriesModified();
+		})
+		.OnSaveAll_Lambda([this]()
+		{
+			SaveAssetEntries();
+			return FReply::Handled();
+		});
 
 	OnOutlinerSelectionChangedDelegate.AddSP(this, &FWorkspaceEditor::HandleOutlinerSelectionChanged);
 	
@@ -335,7 +347,15 @@ TSharedPtr<SDockTab> FWorkspaceEditor::OpenDocument(const UObject* InForObject,
 	const bool bIsSupportedDocument = DocumentManager->FindSupportingFactory(Payload).IsValid();
 	if (bIsSupportedDocument)
 	{
-		AddEditingObject(const_cast<UObject*>(InForObject));
+		const bool bIsRootAsset = Workspace->GetRootAsset() == InForObject;
+		
+		if (bIsRootAsset || Workspace->GetSchema()->SupportsMultipleAssets())
+		{
+			// Only associate this workspace editor with the root level asset
+			// Referenced assets are opened as a 'preview' but should not be used in AddEditingObject
+			// as it will prevents opening that asset in a new workspace editor as a root-level asset
+			AddEditingObject(const_cast<UObject*>(InForObject));
+		}
 	}
 	
 	TSharedPtr<SDockTab> NewTab = DocumentManager->OpenDocument(Payload, InCause);
@@ -696,8 +716,8 @@ void FWorkspaceEditor::ExtendToolbar() const
 {
 	if (UToolMenu* Menu = UToolMenus::Get()->ExtendMenu(GetToolMenuToolbarName()))
 	{
-		FToolMenuSection& WorkspaceOperationsSection = Menu->AddSection("WorkspaceOperations");
-		WorkspaceOperationsSection.AddEntry(FToolMenuEntry::InitToolBarButton(FWorkspaceAssetEditorCommands::Get().SaveAssetEntries, FText::GetEmpty(), TAttribute<FText>(), FSlateIcon(FAppStyle::GetAppStyleSetName(), "MainFrame.SaveAll")));
+		// Keep as other systems hook off of this section
+		Menu->AddSection("WorkspaceOperations");
 	}
 }
 
@@ -774,19 +794,27 @@ void FWorkspaceEditor::InitToolMenuContext(FToolMenuContext& InMenuContext)
 
 void FWorkspaceEditor::SaveAsset_Execute()
 {
-	// If asset is a default 'Untitled' workspace, redirect to the 'save as' flow
-	const FString AssetPath = Workspace->GetOutermost()->GetPathName();
-	if(AssetPath.StartsWith(TEXT("/Temp/Untitled")))
+	if (GetSchema()->CanSaveWorkspace())
 	{
-		// Ensure we do not also 'save as' other externally linked assets at this point
+		// If asset is a default 'Untitled' workspace, redirect to the 'save as' flow
+		const FString AssetPath = Workspace->GetOutermost()->GetPathName();
+		if(AssetPath.StartsWith(TEXT("/Temp/Untitled")))
 		{
-			TGuardValue<bool> SavingTransientWorkspace(bSavingTransientWorkspace, true);
-			SaveAssetAs_Execute();
+			// Ensure we do not also 'save as' other externally linked assets at this point
+			{
+				TGuardValue<bool> SavingTransientWorkspace(bSavingTransientWorkspace, true);
+				SaveAssetAs_Execute();
+			}
+		}
+		else
+		{
+			TGuardValue<bool> SavingWorkspace(bSavingWorkspace, true);
+			FBaseAssetToolkit::SaveAsset_Execute();
 		}
 	}
 	else
 	{
-		TGuardValue<bool> SavingWorkspace(bSavingWorkspace, true);
+		// Saves the assets returned by our GetSaveableObjects implementation
 		FBaseAssetToolkit::SaveAsset_Execute();
 	}
 }
@@ -821,38 +849,61 @@ bool FWorkspaceEditor::InEditingMode() const
 
 void FWorkspaceEditor::GetSaveableObjects(TArray<UObject*>& OutObjects) const
 {
-	if (bSavingWorkspace || bSavingTransientWorkspace)
+	if (GetSchema()->CanSaveWorkspace())
 	{
-		FBaseAssetToolkit::GetSaveableObjects(OutObjects);
-	}
-	
-	if(bSavingAssetEntries)
-	{
-		for(const UWorkspaceAssetEntry* Entry : Workspace->AssetEntries)
+		if (bSavingWorkspace || bSavingTransientWorkspace)
+		{
+			FBaseAssetToolkit::GetSaveableObjects(OutObjects);
+		}
+    	
+		if(bSavingAssetEntries)
 		{
-			if(UObject* Asset = Entry->Asset.Get())
+			for(const UWorkspaceAssetEntry* Entry : Workspace->AssetEntries)
 			{
-				// Add object referenced by workspace
-				OutObjects.Add(Asset);
+				if(UObject* Asset = Entry->Asset.Get())
+				{
+					// Add object referenced by workspace
+					OutObjects.Add(Asset);
+				}
 			}
 		}
 	}
+	else
+	{
+		if (UObject* FocusedAsset = GetFocusedAsset())
+		{
+			OutObjects.Add(FocusedAsset);
+		}
+	}
 }
 
 bool FWorkspaceEditor::CanSaveAsset() const
 {
-	bool bDirtyState = false;
-	for (UObject* EditingObject : GetEditingObjects())
+	if (GetSchema()->CanSaveWorkspace())
 	{
-		const UPackage* Package = EditingObject->GetPackage();
-		if (Package->IsDirty() || Package->GetExternalPackages().ContainsByPredicate([](const UPackage* ExternalPackage)  { return ExternalPackage && ExternalPackage->IsDirty(); }))
+		bool bDirtyState = false;
+		for (UObject* EditingObject : GetEditingObjects())
 		{
-			bDirtyState = true;
-			break;
+			const UPackage* Package = EditingObject->GetPackage();
+			if (Package->IsDirty() || Package->GetExternalPackages().ContainsByPredicate([](const UPackage* ExternalPackage)  { return ExternalPackage && ExternalPackage->IsDirty(); }))
+			{
+				bDirtyState = true;
+				break;
+			}
 		}
+    
+		return bDirtyState;
 	}
+	else
+	{
+		if (const UObject* RootAsset = GetFocusedAsset())
+		{
+			const UPackage* RootAssetPackage = RootAsset->GetPackage();
+			return RootAssetPackage->IsDirty();
+		}
 
-	return bDirtyState;
+		return false;
+	}
 }
 
 FText FWorkspaceEditor::GetTabSuffix() const
@@ -862,26 +913,108 @@ FText FWorkspaceEditor::GetTabSuffix() const
 
 FText FWorkspaceEditor::GetToolkitName() const
 {
-	UObject* const* WorkspaceObject = GetEditingObjects().FindByPredicate([](const UObject* Object) { return Object && Object->IsA<UWorkspace>(); });
-	check (WorkspaceObject != nullptr);
+	if (GetSchema()->CanSaveWorkspace())
+	{
+		UObject* const* WorkspaceObject = GetEditingObjects().FindByPredicate([](const UObject* Object) { return Object && Object->IsA<UWorkspace>(); });
+		check (WorkspaceObject != nullptr);
+    	
+		return GetLabelForObject(*WorkspaceObject);
+	}
 	
-	return GetLabelForObject(*WorkspaceObject);
+	if (const UObject* RootAsset = Workspace->GetRootAsset())
+	{
+		return GetLabelForObject(RootAsset);
+	}
+
+	if (const UObject* FocusedAsset = GetFocusedAsset())
+	{
+		return GetLabelForObject(FocusedAsset);
+	}
+
+	return LOCTEXT("Workspace", "Workspace");
 }
 
 FText FWorkspaceEditor::GetToolkitToolTipText() const
 {
-	UObject* const* WorkspaceObject = GetEditingObjects().FindByPredicate([](const UObject* Object) { return Object && Object->IsA<UWorkspace>(); });
-	check (WorkspaceObject != nullptr);
+	if (GetSchema()->CanSaveWorkspace())
+	{
+		UObject* const* WorkspaceObject = GetEditingObjects().FindByPredicate([](const UObject* Object) { return Object && Object->IsA<UWorkspace>(); });
+		check (WorkspaceObject != nullptr);
+		
+		FText FocusedAssetText = FText::FromName(NAME_None);
+	
+		const FWorkspaceDocument& Document = GetFocusedWorkspaceDocument();
+		if (const UObject* FocusedObject = Document.GetObject())
+		{
+			 FocusedAssetText = GetLabelForObject(FocusedObject);
+		}
+		
+		return FText::Format(LOCTEXT("ToolkitTooltipFormat", "{0} ({1})"), GetToolTipTextForObject(*WorkspaceObject), FocusedAssetText);
+	}
+	
+	if (const UObject* RootAsset = Workspace->GetRootAsset())
+	{
+		return GetLabelForObject(RootAsset);
+	}
 	
-	FText FocusedAssetText = FText::FromName(NAME_None);
+	if (const UObject* FocusedAsset = GetFocusedAsset())
+	{
+		return GetLabelForObject(FocusedAsset);
+	}
+
+	return LOCTEXT("Workspace", "Workspace");
+}
+
+const FSlateBrush* FWorkspaceEditor::GetDefaultTabIcon() const
+{
+	if (GetSchema()->CanSaveWorkspace())
+	{
+		return FBaseAssetToolkit::GetDefaultTabIcon();
+	}
+	
+	if (const UObject* RootAsset = Workspace->GetRootAsset())
+    {
+		return FSlateIconFinder::FindIconBrushForClass(RootAsset->GetClass());
+	}
+	
+	if (const UObject* FocusedAsset = GetFocusedAsset())
+	{
+		return FSlateIconFinder::FindIconBrushForClass(FocusedAsset->GetClass());
+	}
+
+	return FAssetEditorToolkit::GetDefaultTabIcon();
+}
+
+FLinearColor FWorkspaceEditor::GetDefaultTabColor() const
+{
+	if (GetSchema()->CanSaveWorkspace())
+	{
+		return FBaseAssetToolkit::GetDefaultTabColor();
+	}
+
+	static const FAssetToolsModule& AssetToolsModule = FModuleManager::LoadModuleChecked<FAssetToolsModule>("AssetTools");
+	
+	if (const UObject* RootAsset = Workspace->GetRootAsset())
+	{
+		const UClass* RootAssetClass = RootAsset->GetClass();
+
+		TWeakPtr<IAssetTypeActions> AssetTypeActions = AssetToolsModule.Get().GetAssetTypeActionsForClass(RootAssetClass);
+		if (AssetTypeActions.IsValid())
+		{
+			return AssetTypeActions.Pin()->GetTypeColor();
+		}
+	}
 
-	const FWorkspaceDocument& Document = GetFocusedWorkspaceDocument();
-	if (const UObject* FocusedObject = Document.GetObject())
+	if (const UObject* FocusedAsset = GetFocusedAsset())
 	{
-		 FocusedAssetText = GetLabelForObject(FocusedObject);
+		TWeakPtr<IAssetTypeActions> AssetTypeActions = AssetToolsModule.Get().GetAssetTypeActionsForClass(FocusedAsset->GetClass());
+		if (AssetTypeActions.IsValid())
+		{
+			return AssetTypeActions.Pin()->GetTypeColor();
+		}
 	}
 	
-	return FText::Format(LOCTEXT("ToolkitTooltipFormat", "{0} ({1})"), GetToolTipTextForObject(*WorkspaceObject), FocusedAssetText);
+    return FLinearColor::White;
 }
 
 void FWorkspaceEditor::RecordDocumentState(const TInstancedStruct<FWorkspaceDocumentState>& InState) const
@@ -941,20 +1074,35 @@ void FWorkspaceEditor::FindInContentBrowser_Execute()
 {
 	if (IsValid(Workspace))
 	{
-		GEditor->SyncBrowserToObject(Workspace);
+		if (Workspace->GetSchema()->CanSaveWorkspace())
+		{
+			GEditor->SyncBrowserToObject(Workspace);
+		}
+		else
+		{
+			if (UObject* FocusedAsset = GetFocusedAsset())
+			{
+				GEditor->SyncBrowserToObject(FocusedAsset);
+			}
+		}
 	}
 }
 
 bool FWorkspaceEditor::OnRequestClose(EAssetEditorCloseReason InCloseReason)
 {
+	if (!GetSchema()->CanSaveWorkspace())
+	{
+		return true;
+	}
+	
 	TGuardValue<bool> ClosingDown(bClosingDown, true);
 
 	auto RequiresSave = [this]()
-	{
-		const UPackage* Package = Workspace->GetOutermost();
-		// Ask the user to save a transient workspace containing more than 1 asset
-		return Package->GetPathName().StartsWith(TEXT("/Temp/Untitled")) && Workspace->AssetEntries.Num() > 1;
-	};
+		{
+			const UPackage* Package = Workspace->GetOutermost();
+			// Ask the user to save a transient workspace containing more than 1 asset
+			return Package->GetPathName().StartsWith(TEXT("/Temp/Untitled")) && Workspace->AssetEntries.Num() > 1;
+		};
 
 	// Give the user opportunity to save temp workspaces
 	if(RequiresSave() && !bSavingTransientWorkspace)
@@ -984,7 +1132,10 @@ bool FWorkspaceEditor::OnRequestClose(EAssetEditorCloseReason InCloseReason)
 
 void FWorkspaceEditor::OnClose()
 {
-	SaveEditedObjectState();
+	if (GetSchema()->CanSaveWorkspace())
+	{
+		SaveEditedObjectState();
+	}
 	
 	if (DetailsView.IsValid())
 	{
@@ -1095,7 +1246,12 @@ UWorkspaceViewportSceneDescription* FWorkspaceEditor::GetSceneDescription() cons
 {
 	return SceneDescription;
 }
-	
+
+UObject* FWorkspaceEditor::GetFocusedAsset() const
+{
+	return GetFocusedWorkspaceDocument().Export.GetFirstAssetOfType<UObject>();
+}
+
 } // namespace UE::Workspace
 
 #undef LOCTEXT_NAMESPACE
diff --git a/Engine/Plugins/Experimental/Workspace/Source/WorkspaceEditor/Private/WorkspaceEditor.h b/Engine/Plugins/Experimental/Workspace/Source/WorkspaceEditor/Private/WorkspaceEditor.h
--- a/Engine/Plugins/Experimental/Workspace/Source/WorkspaceEditor/Private/WorkspaceEditor.h
+++ b/Engine/Plugins/Experimental/Workspace/Source/WorkspaceEditor/Private/WorkspaceEditor.h
@@ -82,6 +82,8 @@ private:
 	virtual FText GetTabSuffix() const override;
 	virtual FText GetToolkitName() const override;
 	virtual FText GetToolkitToolTipText() const override;
+	virtual const FSlateBrush* GetDefaultTabIcon() const override;
+    virtual FLinearColor GetDefaultTabColor() const override;
 	virtual void FindInContentBrowser_Execute() override;
 	virtual void OnToolkitHostingFinished(const TSharedRef<IToolkit>& Toolkit) override;
 	virtual void OnToolkitHostingStarted(const TSharedRef<IToolkit>& Toolkit) override;
@@ -147,6 +149,8 @@ private:
 	bool AreAssetEntriesModified() const;
 
 	bool IsViewportPinned() const;
+	
+	UObject* GetFocusedAsset() const;
 
 	/** The asset being edited */
 	TObjectPtr<UWorkspace> Workspace = nullptr;
diff --git a/Engine/Plugins/Experimental/Workspace/Source/WorkspaceEditor/Private/WorkspaceEditorModule.cpp b/Engine/Plugins/Experimental/Workspace/Source/WorkspaceEditor/Private/WorkspaceEditorModule.cpp
--- a/Engine/Plugins/Experimental/Workspace/Source/WorkspaceEditor/Private/WorkspaceEditorModule.cpp
+++ b/Engine/Plugins/Experimental/Workspace/Source/WorkspaceEditor/Private/WorkspaceEditorModule.cpp
@@ -20,6 +20,7 @@
 #include "Styling/AppStyle.h"
 #include "WorkspaceAssetReferenceItemDetails.h"
 #include "WorkspaceGroupItemDetails.h"
+#include "WorkspaceSchema.h"
 
 #define LOCTEXT_NAMESPACE "WorkspaceEditorModule"
 
@@ -214,94 +215,118 @@ bool FWorkspaceEditorModule::GetExportedAssetsForWorkspace(const FAssetData& InW
 
  IWorkspaceEditor* FWorkspaceEditorModule::OpenWorkspaceForObject(UObject* InObject, EOpenWorkspaceMethod InOpenMethod, const TSubclassOf<UWorkspaceFactory> WorkSpaceFactoryClass/*= UWorkspaceFactory::StaticClass()*/)
 {
-	FAssetRegistryModule& AssetRegistryModule = FModuleManager::LoadModuleChecked<FAssetRegistryModule>(TEXT("AssetRegistry"));
-
-	TArray<FAssetData> RelevantWorkspaceAssets;
-
-	if(InOpenMethod != EOpenWorkspaceMethod::AlwaysOpenNewWorkspace)
+	UWorkspaceFactory* DefaultWorkspaceFactory = WorkSpaceFactoryClass->GetDefaultObject<UWorkspaceFactory>();
+	UWorkspaceSchema* DefaultWorkspaceSchema = DefaultWorkspaceFactory->SchemaClass->GetDefaultObject<UWorkspaceSchema>();
+	
+	if (DefaultWorkspaceSchema->CanSaveWorkspace())
 	{
-		// Look for existing workspaces that export this asset
-		FARFilter ARFilter;
-		ARFilter.ClassPaths.Add(UWorkspace::StaticClass()->GetClassPathName());
-		ARFilter.bRecursiveClasses = true;
-
-		TArray<FAssetData> AllWorkspaceAssets;
-		AssetRegistryModule.Get().GetAssets(ARFilter, AllWorkspaceAssets);
-
-		for(const FAssetData& WorkspaceAsset : AllWorkspaceAssets)
+		FAssetRegistryModule& AssetRegistryModule = FModuleManager::LoadModuleChecked<FAssetRegistryModule>(TEXT("AssetRegistry"));
+	
+		TArray<FAssetData> RelevantWorkspaceAssets;
+	
+		if(InOpenMethod != EOpenWorkspaceMethod::AlwaysOpenNewWorkspace)
 		{
-			FWorkspaceAssetRegistryExports Exports;
-			GetExportedAssetsForWorkspace(WorkspaceAsset, Exports);
-
-			FSoftObjectPath ObjectPath(InObject);
-			for(const FWorkspaceAssetRegistryExportEntry& ExportEntry : Exports.Assets)
+			// Look for existing workspaces that export this asset
+			FARFilter ARFilter;
+			ARFilter.ClassPaths.Add(UWorkspace::StaticClass()->GetClassPathName());
+			ARFilter.bRecursiveClasses = true;
+	
+			TArray<FAssetData> AllWorkspaceAssets;
+			AssetRegistryModule.Get().GetAssets(ARFilter, AllWorkspaceAssets);
+	
+			for(const FAssetData& WorkspaceAsset : AllWorkspaceAssets)
 			{
-				if(ExportEntry.Asset == ObjectPath)
+				FWorkspaceAssetRegistryExports Exports;
+				GetExportedAssetsForWorkspace(WorkspaceAsset, Exports);
+	
+				FSoftObjectPath ObjectPath(InObject);
+				for(const FWorkspaceAssetRegistryExportEntry& ExportEntry : Exports.Assets)
 				{
-					RelevantWorkspaceAssets.Add(WorkspaceAsset);
-					break;
+					if(ExportEntry.Asset == ObjectPath)
+					{
+						RelevantWorkspaceAssets.Add(WorkspaceAsset);
+						break;
+					}
 				}
 			}
 		}
+	
+		IWorkspaceEditor* WorkspaceEditor = nullptr;
+	
+		auto HandleNewWorkspace = [InObject, &WorkspaceEditor](TObjectPtr<UWorkspace> NewWorkspace)
+		{
+			NewWorkspace->AddAsset(InObject, false);
+			NewWorkspace->MarkPackageDirty();
+	
+			UAssetEditorSubsystem* AssetEditorSubsystem = GEditor->GetEditorSubsystem<UAssetEditorSubsystem>();
+			UWorkspaceAssetEditor* AssetEditor = NewObject<UWorkspaceAssetEditor>(AssetEditorSubsystem, NAME_None, RF_Transient);
+			AssetEditor->SetObjectToEdit(NewWorkspace);
+			AssetEditor->Initialize();
+			
+			WorkspaceEditor = static_cast<IWorkspaceEditor*>(GEditor->GetEditorSubsystem<UAssetEditorSubsystem>()->FindEditorForAsset(NewWorkspace, true));
+		};
+	
+		auto HandleExistingWorkspace = [InObject, &WorkspaceEditor](TObjectPtr<UWorkspace> ExistingWorkspace)
+		{
+			GEditor->GetEditorSubsystem<UAssetEditorSubsystem>()->OpenEditorForAsset(ExistingWorkspace);
+	
+			WorkspaceEditor = static_cast<IWorkspaceEditor*>(GEditor->GetEditorSubsystem<UAssetEditorSubsystem>()->FindEditorForAsset(ExistingWorkspace, true));
+		};
+	
+		if(InOpenMethod == EOpenWorkspaceMethod::AlwaysOpenNewWorkspace || RelevantWorkspaceAssets.Num() == 0)
+		{
+			// No relevant workspaces, so open a new one and add the asset
+			UWorkspaceFactory* Factory = NewObject<UWorkspaceFactory>(GetTransientPackage(), WorkSpaceFactoryClass.Get());
+			UPackage* Package = CreatePackage(nullptr);
+			FName PackageName = *FPaths::GetBaseFilename(Package->GetName());
+			if (Factory->ConfigureProperties())
+			{
+				TObjectPtr<UWorkspace> SelectedWorkspace = CastChecked<UWorkspace>(Factory->FactoryCreateNew(UWorkspace::StaticClass(), Package, PackageName, RF_Public | RF_Standalone, nullptr, GWarn));
+				HandleNewWorkspace(SelectedWorkspace);
+			}
+		}
+		else if(RelevantWorkspaceAssets.Num() == 1)
+		{
+			// One existing workspace, open it
+			HandleExistingWorkspace(CastChecked<UWorkspace>(RelevantWorkspaceAssets[0].GetAsset()));
+		}
+		else
+		{
+			// Multiple existing workspaces, present a window to let the user choose one to open with
+			TSharedRef<SWorkspacePicker> WorkspacePicker = SNew(SWorkspacePicker)
+				.WorkspaceAssets(RelevantWorkspaceAssets)
+				.OnExistingWorkspaceSelected_Lambda(HandleExistingWorkspace)
+				.OnNewWorkspaceCreated_Lambda(HandleNewWorkspace)
+				.WorkspaceFactoryClass(WorkSpaceFactoryClass.Get());
+	
+			WorkspacePicker->ShowModal();
+		}
+	
+		if(WorkspaceEditor)
+		{
+			WorkspaceEditor->OpenAssets({InObject});
+		}
+	
+		return WorkspaceEditor;
 	}
-
-	IWorkspaceEditor* WorkspaceEditor = nullptr;
-
-	auto HandleNewWorkspace = [InObject, &WorkspaceEditor](TObjectPtr<UWorkspace> NewWorkspace)
+	else
 	{
+		UWorkspaceFactory* Factory = NewObject<UWorkspaceFactory>(GetTransientPackage(), WorkSpaceFactoryClass.Get());
+		TObjectPtr<UWorkspace> NewWorkspace = CastChecked<UWorkspace>(Factory->FactoryCreateNew(UWorkspace::StaticClass(), GetTransientPackage(), NAME_None, RF_Public | RF_Standalone, nullptr, GWarn));
+		
 		NewWorkspace->AddAsset(InObject, false);
-		NewWorkspace->MarkPackageDirty();
-
+		(void)NewWorkspace->MarkPackageDirty();
+	
 		UAssetEditorSubsystem* AssetEditorSubsystem = GEditor->GetEditorSubsystem<UAssetEditorSubsystem>();
 		UWorkspaceAssetEditor* AssetEditor = NewObject<UWorkspaceAssetEditor>(AssetEditorSubsystem, NAME_None, RF_Transient);
 		AssetEditor->SetObjectToEdit(NewWorkspace);
 		AssetEditor->Initialize();
 		
-		WorkspaceEditor = static_cast<IWorkspaceEditor*>(GEditor->GetEditorSubsystem<UAssetEditorSubsystem>()->FindEditorForAsset(NewWorkspace, true));
-	};
-
-	auto HandleExistingWorkspace = [InObject, &WorkspaceEditor](TObjectPtr<UWorkspace> ExistingWorkspace)
-	{
-		GEditor->GetEditorSubsystem<UAssetEditorSubsystem>()->OpenEditorForAsset(ExistingWorkspace);
-
-		WorkspaceEditor = static_cast<IWorkspaceEditor*>(GEditor->GetEditorSubsystem<UAssetEditorSubsystem>()->FindEditorForAsset(ExistingWorkspace, true));
-	};
-
-	if(InOpenMethod == EOpenWorkspaceMethod::AlwaysOpenNewWorkspace || RelevantWorkspaceAssets.Num() == 0)
-	{
-		// No relevant workspaces, so open a new one and add the asset
-		UWorkspaceFactory* Factory = NewObject<UWorkspaceFactory>(GetTransientPackage(), WorkSpaceFactoryClass.Get());
-		UPackage* Package = CreatePackage(nullptr);
-		FName PackageName = *FPaths::GetBaseFilename(Package->GetName());
-		if (Factory->ConfigureProperties())
-		{
-			TObjectPtr<UWorkspace> SelectedWorkspace = CastChecked<UWorkspace>(Factory->FactoryCreateNew(UWorkspace::StaticClass(), Package, PackageName, RF_Public | RF_Standalone, nullptr, GWarn));
-			HandleNewWorkspace(SelectedWorkspace);
-		}
-	}
-	else if(RelevantWorkspaceAssets.Num() == 1)
-	{
-		// One existing workspace, open it
-		HandleExistingWorkspace(CastChecked<UWorkspace>(RelevantWorkspaceAssets[0].GetAsset()));
-	}
-	else
-	{
-		// Multiple existing workspaces, present a window to let the user choose one to open with
-		TSharedRef<SWorkspacePicker> WorkspacePicker = SNew(SWorkspacePicker)
-			.WorkspaceAssets(RelevantWorkspaceAssets)
-			.OnExistingWorkspaceSelected_Lambda(HandleExistingWorkspace)
-			.OnNewWorkspaceCreated_Lambda(HandleNewWorkspace)
-			.WorkspaceFactoryClass(WorkSpaceFactoryClass.Get());
-
-		WorkspacePicker->ShowModal();
-	}
-
-	if(WorkspaceEditor)
-	{
+		IWorkspaceEditor* WorkspaceEditor = static_cast<IWorkspaceEditor*>(GEditor->GetEditorSubsystem<UAssetEditorSubsystem>()->FindEditorForAsset(NewWorkspace, true));
 		WorkspaceEditor->OpenAssets({InObject});
+		
+		return WorkspaceEditor;
 	}
-
-	return WorkspaceEditor;
 }
 
  TSharedPtr<IWorkspacePicker> FWorkspaceEditorModule::CreateWorkspacePicker(const IWorkspacePicker::FConfig& Config) const
diff --git a/Engine/Plugins/Experimental/Workspace/Source/WorkspaceEditor/Public/WorkspaceSchema.h b/Engine/Plugins/Experimental/Workspace/Source/WorkspaceEditor/Public/WorkspaceSchema.h
--- a/Engine/Plugins/Experimental/Workspace/Source/WorkspaceEditor/Public/WorkspaceSchema.h
+++ b/Engine/Plugins/Experimental/Workspace/Source/WorkspaceEditor/Public/WorkspaceSchema.h
@@ -38,4 +38,13 @@ public:
 
 	// Create a new object containing the required data to construct a scene for this type of workspace
 	virtual TObjectPtr<UWorkspaceViewportSceneDescription> CreateSceneDescription() const { return NewObject<UWorkspaceViewportSceneDescription>(); }
+	
+	// Whether workspaces using this schema can add more assets to the workspace besides the root-level asset
+	virtual bool SupportsMultipleAssets() const { return false; }
+	
+	// Whether workspaces using this schema has a root-level asset that always exists in the workspace
+	virtual bool HasPersistentRootAsset() const { return false; }
+	
+	// Whether workspaces using this schema can be saved as an asset
+	virtual bool CanSaveWorkspace() const { return false; }
 };
