Fix transient workspace using external packages

#rnx

diff --git a/Engine/Plugins/Experimental/Workspace/Source/WorkspaceEditor/Private/Workspace.cpp b/Engine/Plugins/Experimental/Workspace/Source/WorkspaceEditor/Private/Workspace.cpp
--- a/Engine/Plugins/Experimental/Workspace/Source/WorkspaceEditor/Private/Workspace.cpp
+++ b/Engine/Plugins/Experimental/Workspace/Source/WorkspaceEditor/Private/Workspace.cpp
@@ -34,11 +34,14 @@ bool UWorkspace::AddAsset(const FAssetData& InAsset, bool bSetupUndoRedo, bool b
 		Modify(false);
 	}
 
-	int32 NewIndex = INDEX_NONE; 
-	if (!AssetEntries.ContainsByPredicate([&InAsset](const UWorkspaceAssetEntry* AssetEntry)-> bool
-	{
-		return AssetEntry->Asset.ToSoftObjectPath() == InAsset.ToSoftObjectPath();
-	}))
+	int32 NewIndex = INDEX_NONE;
+
+	const bool bAssetNotInWorkspace = !AssetEntries.ContainsByPredicate([&InAsset](const UWorkspaceAssetEntry* AssetEntry)-> bool
+		{
+			return AssetEntry->Asset.ToSoftObjectPath() == InAsset.ToSoftObjectPath();
+		});
+
+	if (bAssetNotInWorkspace)
 	{
 		// If we are a transient asset, dont use external packages
 		const UObject* Asset = InAsset.GetAsset();
@@ -46,7 +49,14 @@ bool UWorkspace::AddAsset(const FAssetData& InAsset, bool bSetupUndoRedo, bool b
 		if(!Asset->HasAnyFlags(RF_Transient))
 		{
 			UWorkspaceAssetEntry* NewEntry = NewObject<UWorkspaceAssetEntry>(this, UWorkspaceAssetEntry::StaticClass(), NAME_None, RF_Transactional);
-			FExternalPackageHelper::SetPackagingMode(NewEntry, this, true, false, PKG_None);
+
+			// Only use external packaging when this workspace is savable otherwise
+			// we want assets housed in our transient workspace.
+			if (GetSchema()->CanSaveWorkspace())
+			{
+				check(!HasAnyFlags(RF_Transient));
+				FExternalPackageHelper::SetPackagingMode(NewEntry, this, true, false, PKG_None);
+			}
 
 			NewEntry->Asset = TSoftObjectPtr<UObject>(InAsset.GetSoftObjectPath());
 			NewIndex = AssetEntries.Add(NewEntry);
