Hide the workspace from the focused document breadcrumps

diff --git a/Engine/Plugins/Experimental/Workspace/Source/WorkspaceEditor/Private/SWorkspaceTabWrapper.cpp b/Engine/Plugins/Experimental/Workspace/Source/WorkspaceEditor/Private/SWorkspaceTabWrapper.cpp
--- a/Engine/Plugins/Experimental/Workspace/Source/WorkspaceEditor/Private/SWorkspaceTabWrapper.cpp
+++ b/Engine/Plugins/Experimental/Workspace/Source/WorkspaceEditor/Private/SWorkspaceTabWrapper.cpp
@@ -7,6 +7,7 @@
 #include "Widgets/Docking/SDockTab.h"
 #include "Modules/ModuleManager.h"
 #include "ToolMenus.h"
+#include "WorkspaceSchema.h"
 
 #define LOCTEXT_NAMESPACE "SWorkspaceTabWrapper"
 
@@ -321,7 +322,7 @@ EVisibility SWorkspaceTabWrapper::IsWorkspaceNameVisible() const
 {
 	if (const TSharedPtr<UE::Workspace::FWorkspaceEditor> SharedWorkspaceEditor = WeakWorkspaceEditor.Pin())
 	{
-		if (SharedWorkspaceEditor->Workspace)
+		if (SharedWorkspaceEditor->Workspace && SharedWorkspaceEditor->GetSchema()->CanSaveWorkspace())
 		{
 			return EVisibility::Visible;
 		}
