Fix color overrides on nested workspace outliner entries

diff --git a/Engine/Plugins/Experimental/Workspace/Source/WorkspaceEditor/Private/WorkspaceAssetReferenceItemDetails.h b/Engine/Plugins/Experimental/Workspace/Source/WorkspaceEditor/Private/WorkspaceAssetReferenceItemDetails.h
--- a/Engine/Plugins/Experimental/Workspace/Source/WorkspaceEditor/Private/WorkspaceAssetReferenceItemDetails.h
+++ b/Engine/Plugins/Experimental/Workspace/Source/WorkspaceEditor/Private/WorkspaceAssetReferenceItemDetails.h
@@ -59,11 +59,19 @@ public:
 				
 		return Icon ? Icon : FAppStyle::GetBrush(TEXT("GenericLink"));
 	}
-
+	
 	virtual FSlateColor GetItemColor(const FWorkspaceOutlinerItemExport& Export) const override
 	{
 		const FWorkspaceOutlinerAssetReferenceItemData& Data = Export.GetData().Get<FWorkspaceOutlinerAssetReferenceItemData>();	
-		return Data.bRecursiveReference && Data.bShouldExpandReference ? FLinearColor::Red : FSlateColor::UseSubduedForeground();
+		
+		FSlateColor InnerColor = FSlateColor::UseSubduedForeground();
+		if (const TSharedPtr<IWorkspaceOutlinerItemDetails> SharedFactory = GetInnerDetails(Export))
+		{
+			const FWorkspaceOutlinerItemExport InnerExport = GetInnerExport(Export);
+			InnerColor = SharedFactory->GetItemColor(InnerExport);
+		}
+		
+		return Data.bRecursiveReference && Data.bShouldExpandReference ? FLinearColor::Red : InnerColor;
 	}	
 
 	virtual bool IsExpandedByDefault() const override
