Update workspace outliner label to 'Outliner' and swap viewport and outliner tab locations

diff --git a/Engine/Plugins/Experimental/Workspace/Source/WorkspaceEditor/Private/WorkspaceEditor.cpp b/Engine/Plugins/Experimental/Workspace/Source/WorkspaceEditor/Private/WorkspaceEditor.cpp
--- a/Engine/Plugins/Experimental/Workspace/Source/WorkspaceEditor/Private/WorkspaceEditor.cpp
+++ b/Engine/Plugins/Experimental/Workspace/Source/WorkspaceEditor/Private/WorkspaceEditor.cpp
@@ -128,7 +128,7 @@ void FWorkspaceEditor::CreateWidgets()
 	check(DetailsView.IsValid());
 	WorkspaceEditorModule.ApplyWorkspaceDetailsCustomization(StaticCastWeakPtr<IWorkspaceEditor>(this->AsWeak()), DetailsView);
 
-	StandaloneDefaultLayout = FTabManager::NewLayout("Standalone_WorkspaceEditor_Layout_v4")
+	StandaloneDefaultLayout = FTabManager::NewLayout("Standalone_WorkspaceEditor_Layout_v5")
 	->AddArea
 	(
 		FTabManager::NewPrimaryArea()
@@ -148,7 +148,7 @@ void FWorkspaceEditor::CreateWidgets()
 					->SetSizeCoefficient(0.5f)
 					->SetHideTabWell(false)
 					->AddTab(WorkspaceTabs::TopLeftDocumentArea, ETabState::ClosedTab)
-					->AddTab(FBaseAssetToolkit::ViewportTabID, ETabState::OpenedTab)
+					->AddTab(WorkspaceTabs::WorkspaceView, ETabState::OpenedTab)
 					->SetExtensionId(WorkspaceTabs::TopLeftDocumentArea)
 				)
 				->Split
@@ -192,7 +192,7 @@ void FWorkspaceEditor::CreateWidgets()
 					FTabManager::NewStack()
 					->SetSizeCoefficient(0.5f)
 					->SetHideTabWell(false)
-					->AddTab(WorkspaceTabs::WorkspaceView, ETabState::OpenedTab)
+					->AddTab(FBaseAssetToolkit::ViewportTabID, ETabState::OpenedTab)
 					->AddTab(WorkspaceTabs::TopRightDocumentArea, ETabState::ClosedTab)
 					->SetExtensionId(WorkspaceTabs::TopRightDocumentArea)
 				)
@@ -732,12 +732,12 @@ void FWorkspaceEditor::RegisterTabSpawners(const TSharedRef<FTabManager>& InTabM
 		check(Args.GetTabId() == WorkspaceTabs::WorkspaceView);
 		
 		return SNew(SDockTab)
-			.Label(LOCTEXT("WorkspaceTabLabel", "Workspace"))
+			.Label(LOCTEXT("WorkspaceTabLabel", "Outliner"))
 			[
 				WorkspaceView.ToSharedRef()
 			];
 	}))
-	.SetDisplayName(LOCTEXT("WorkspaceTabLabel", "Workspace"))
+	.SetDisplayName(LOCTEXT("WorkspaceTabLabel", "Outliner"))
 	.SetIcon(FSlateIcon("EditorStyle", "LevelEditor.Tabs.Outliner"))
 	.SetTooltipText(LOCTEXT("WorkspaceTabToolTip", "Shows the workspace outliner tab."));
 
