Fix calls into FWorkspaceEditorModule::OpenWorkspaceForObject without a valid schema setup by asking the user to pick one or early out if no valid one was selected

#Synced-CL 48348137

diff --git a/Engine/Plugins/Experimental/Workspace/Source/WorkspaceEditor/Private/WorkspaceEditorModule.cpp b/Engine/Plugins/Experimental/Workspace/Source/WorkspaceEditor/Private/WorkspaceEditorModule.cpp
--- a/Engine/Plugins/Experimental/Workspace/Source/WorkspaceEditor/Private/WorkspaceEditorModule.cpp
+++ b/Engine/Plugins/Experimental/Workspace/Source/WorkspaceEditor/Private/WorkspaceEditorModule.cpp
@@ -216,9 +216,9 @@ bool FWorkspaceEditorModule::GetExportedAssetsForWorkspace(const FAssetData& InW
  IWorkspaceEditor* FWorkspaceEditorModule::OpenWorkspaceForObject(UObject* InObject, EOpenWorkspaceMethod InOpenMethod, const TSubclassOf<UWorkspaceFactory> WorkSpaceFactoryClass/*= UWorkspaceFactory::StaticClass()*/)
 {
 	UWorkspaceFactory* DefaultWorkspaceFactory = WorkSpaceFactoryClass->GetDefaultObject<UWorkspaceFactory>();
-	UWorkspaceSchema* DefaultWorkspaceSchema = DefaultWorkspaceFactory->SchemaClass->GetDefaultObject<UWorkspaceSchema>();
-	
-	if (DefaultWorkspaceSchema->CanSaveWorkspace())
+	UWorkspaceSchema* DefaultWorkspaceSchema = DefaultWorkspaceFactory->SchemaClass ? DefaultWorkspaceFactory->SchemaClass->GetDefaultObject<UWorkspaceSchema>() : nullptr;
+
+	if (DefaultWorkspaceSchema && DefaultWorkspaceSchema->CanSaveWorkspace())
 	{
 		FAssetRegistryModule& AssetRegistryModule = FModuleManager::LoadModuleChecked<FAssetRegistryModule>(TEXT("AssetRegistry"));
 	
@@ -312,6 +312,14 @@ bool FWorkspaceEditorModule::GetExportedAssetsForWorkspace(const FAssetData& InW
 	else
 	{
 		UWorkspaceFactory* Factory = NewObject<UWorkspaceFactory>(GetTransientPackage(), WorkSpaceFactoryClass.Get());
+		if (!Factory->SchemaClass)
+		{
+			if (!ensureMsgf(Factory->ConfigureProperties(), TEXT("Can not create a workspace without a valid Schema")))
+			{
+				return nullptr;
+			}
+		}
+
 		TObjectPtr<UWorkspace> NewWorkspace = CastChecked<UWorkspace>(Factory->FactoryCreateNew(UWorkspace::StaticClass(), GetTransientPackage(), NAME_None, RF_Public | RF_Standalone, nullptr, GWarn));
 		
 		NewWorkspace->AddAsset(InObject, false);
