UI command error when selecting a trait stack node and then opening the trait editor
#fix add commandlist bindings for node alignment/distribution/straightening etc.
#Synced-CL 48351162

diff --git a/Engine/Plugins/Experimental/Workspace/Source/WorkspaceEditor/Private/SGraphDocument.cpp b/Engine/Plugins/Experimental/Workspace/Source/WorkspaceEditor/Private/SGraphDocument.cpp
--- a/Engine/Plugins/Experimental/Workspace/Source/WorkspaceEditor/Private/SGraphDocument.cpp
+++ b/Engine/Plugins/Experimental/Workspace/Source/WorkspaceEditor/Private/SGraphDocument.cpp
@@ -16,7 +16,7 @@ void SGraphDocument::Construct(const FArguments& InArgs, TSharedRef<FWorkspaceEd
 	HostingAppPtr = InHostingApp;
 	EdGraph = InWorkspaceDocument.GetTypedObjectChecked<UEdGraph>();
 	Document = InWorkspaceDocument;
-
+	
 	BindCommands();
 
 	OnCanDeleteSelectedNodes = InArgs._OnCanDeleteSelectedNodes;
@@ -107,7 +107,18 @@ void SGraphDocument::BindCommands()
 
 	CommandList->MapAction(FGraphEditorCommands::Get().CreateComment,
 		FExecuteAction::CreateSP(this, &SGraphDocument::CreateComment));
-
+		
+	TWeakPtr<SGraphDocument> WeakDocument = StaticCastSharedRef<SGraphDocument>(AsShared());
+	CommandList->MapAction(FGraphEditorCommands::Get().AlignNodesTop, FExecuteAction::CreateSPLambda(this, [WeakDocument](){ if (TSharedPtr<SGraphDocument> SharedDocument = WeakDocument.Pin()) { SharedDocument->GraphEditor->OnAlignTop(); }}));
+	CommandList->MapAction(FGraphEditorCommands::Get().AlignNodesTop, FExecuteAction::CreateSPLambda(this, [WeakDocument](){ if (TSharedPtr<SGraphDocument> SharedDocument = WeakDocument.Pin()) { SharedDocument->GraphEditor->OnAlignTop(); }}));
+	CommandList->MapAction(FGraphEditorCommands::Get().AlignNodesMiddle, FExecuteAction::CreateSPLambda(this, [WeakDocument](){ if (TSharedPtr<SGraphDocument> SharedDocument = WeakDocument.Pin()) { SharedDocument->GraphEditor->OnAlignMiddle(); }}));
+	CommandList->MapAction(FGraphEditorCommands::Get().AlignNodesBottom, FExecuteAction::CreateSPLambda(this, [WeakDocument](){ if (TSharedPtr<SGraphDocument> SharedDocument = WeakDocument.Pin()) { SharedDocument->GraphEditor->OnAlignBottom(); }}));
+	CommandList->MapAction(FGraphEditorCommands::Get().AlignNodesLeft, FExecuteAction::CreateSPLambda(this, [WeakDocument](){ if (TSharedPtr<SGraphDocument> SharedDocument = WeakDocument.Pin()) { SharedDocument->GraphEditor->OnAlignLeft(); }}));
+	CommandList->MapAction(FGraphEditorCommands::Get().AlignNodesCenter, FExecuteAction::CreateSPLambda(this, [WeakDocument](){ if (TSharedPtr<SGraphDocument> SharedDocument = WeakDocument.Pin()) { SharedDocument->GraphEditor->OnAlignCenter(); }}));
+	CommandList->MapAction(FGraphEditorCommands::Get().AlignNodesRight, FExecuteAction::CreateSPLambda(this, [WeakDocument](){ if (TSharedPtr<SGraphDocument> SharedDocument = WeakDocument.Pin()) { SharedDocument->GraphEditor->OnAlignRight(); }}));
+	CommandList->MapAction(FGraphEditorCommands::Get().StraightenConnections, FExecuteAction::CreateSPLambda(this, [WeakDocument](){ if (TSharedPtr<SGraphDocument> SharedDocument = WeakDocument.Pin()) { SharedDocument->GraphEditor->OnStraightenConnections(); }}));
+	CommandList->MapAction(FGraphEditorCommands::Get().DistributeNodesHorizontally, FExecuteAction::CreateSPLambda(this, [WeakDocument](){ if (TSharedPtr<SGraphDocument> SharedDocument = WeakDocument.Pin()) { SharedDocument->GraphEditor->OnDistributeNodesH(); }}));
+	CommandList->MapAction(FGraphEditorCommands::Get().DistributeNodesVertically, FExecuteAction::CreateSPLambda(this, [WeakDocument](){ if (TSharedPtr<SGraphDocument> SharedDocument = WeakDocument.Pin()) { SharedDocument->GraphEditor->OnDistributeNodesV(); }}));
 }
 
 bool SGraphDocument::CanDeleteSelectedNodes() const
@@ -224,4 +235,4 @@ void SGraphDocument::OnResetSelection()
 {
 	GraphEditor->ClearSelectionSet();
 }
-};
\ No newline at end of file
+};
