Improve workspace view performance
- Cache FAssetData for AssetEntries (refreshed on full outliner refresh/reload)
- Rely on GetAssetByObjectPath rather than AR filtering (very significant speedup)
- Return FAssetData as ArrayView rather than populating in-out TArray

#Synced-CL 48648404

diff --git a/Engine/Plugins/Experimental/Workspace/Source/WorkspaceEditor/Private/Outliner/WorkspaceOutlinerHierarchy.cpp b/Engine/Plugins/Experimental/Workspace/Source/WorkspaceEditor/Private/Outliner/WorkspaceOutlinerHierarchy.cpp
--- a/Engine/Plugins/Experimental/Workspace/Source/WorkspaceEditor/Private/Outliner/WorkspaceOutlinerHierarchy.cpp
+++ b/Engine/Plugins/Experimental/Workspace/Source/WorkspaceEditor/Private/Outliner/WorkspaceOutlinerHierarchy.cpp
@@ -54,24 +54,14 @@ namespace UE::Workspace
 
 				if (bIsReference)
 				{
-					const FWorkspaceOutlinerAssetReferenceItemData& ReferenceItemData = Export.GetData().Get<FWorkspaceOutlinerAssetReferenceItemData>();		
+					const FWorkspaceOutlinerAssetReferenceItemData& ReferenceItemData = Export.GetData().Get<FWorkspaceOutlinerAssetReferenceItemData>();
 
 					// Retrieve AssetRegistry data for the referred asset
-					FARFilter Filter;
-					Filter.SoftObjectPaths.Add(ReferenceItemData.ReferredObjectPath);
-					const IAssetRegistry& AssetRegistry = FAssetRegistryModule::GetRegistry();
-					TArray<FAssetData> OutAssetData;
-					AssetRegistry.GetAssets(Filter, OutAssetData);
+					FAssetData ReferredAssetData = IAssetRegistry::GetChecked().GetAssetByObjectPath(ReferenceItemData.ReferredObjectPath);
 
-					if (OutAssetData.Num())
+					if (ReferredAssetData.IsValid())
 					{
-						// Only ever expect one to be found
-						check(OutAssetData.Num() == 1);
-					
-						const FAssetData& ReferredAssetData = OutAssetData[0];
 						TArray<FSceneOutlinerTreeItemPtr> ReferredItems;
-
-
 						// Make local copy of AssetData array for recursion checks
 						TArray<FAssetData> ReferenceAssets = OutAssets;
 						ReferenceAssets.Add(AssetData);
@@ -80,7 +70,7 @@ namespace UE::Workspace
 						{
 							// Asset reference not at root-level of the workspace, meaning its root entry has to be reparented (in case the find fails, the asset registry exporting is not being done in the same order as outliner representation)
 							// Populate export to-be-used by root-level asset referenced export 
-							const FWorkspaceOutlinerItemExport ReparentedExport(Export.GetIdentifier(), ReparentedHashes.FindChecked(Export.GetParentHash()), Export.GetData());						
+							const FWorkspaceOutlinerItemExport ReparentedExport(Export.GetIdentifier(), ReparentedHashes.FindChecked(Export.GetParentHash()), Export.GetData());
 							ReparentedHashes.Add(GetTypeHash(Export), ReparentedExport);
 							CreateItemsFromAssetData(ReferredAssetData, ReferredItems, ReferenceAssets, &ReparentedExport);
 						}
@@ -134,10 +124,8 @@ namespace UE::Workspace
 				}
 				else
 				{
-					AddItem(Export);				
+					AddItem(Export);
 				}
-
-				
 			}
 			
 			OutAssets.Add(AssetData);
@@ -148,8 +136,7 @@ namespace UE::Workspace
 	{
 		if (const UWorkspace* Workspace = WeakWorkspace.Get())
 		{
-			TArray<FAssetData> AssetDataEntries; 
-			Workspace->GetAssetDataEntries(AssetDataEntries);
+			TConstArrayView<FAssetData> AssetDataEntries = Workspace->GetAssetDataEntriesView(true);
 			ProcessedAssetData.Reset();
 
 			for (const FAssetData& AssetData : AssetDataEntries)
@@ -157,7 +144,7 @@ namespace UE::Workspace
 				TArray<FAssetData> Assets;
 				CreateItemsFromAssetData(AssetData, OutItems, Assets, nullptr);
 			}
-		}		
+		}
 	}
 
 	FSceneOutlinerTreeItemPtr FWorkspaceOutlinerHierarchy::FindOrCreateParentItem(const ISceneOutlinerTreeItem& Item, const TMap<FSceneOutlinerTreeItemID, FSceneOutlinerTreeItemPtr>& Items, bool bCreate)
@@ -175,9 +162,7 @@ namespace UE::Workspace
 				{
 					if (const UWorkspace* Workspace = WeakWorkspace.Get())
 					{
-						TArray<FAssetData> AssetDataEntries; 
-						Workspace->GetAssetDataEntries(AssetDataEntries);
-
+						TConstArrayView<FAssetData> AssetDataEntries = Workspace->GetAssetDataEntriesView();
 						if(const FAssetData* AssetDataPtr = AssetDataEntries.FindByPredicate([AssetPath = TreeItem->Export.GetTopLevelAssetPath()](const FAssetData& AssetData) { return AssetData.GetSoftObjectPath() == AssetPath; }))
 						{
 							FString TagValue;
diff --git a/Engine/Plugins/Experimental/Workspace/Source/WorkspaceEditor/Private/SWorkspaceView.cpp b/Engine/Plugins/Experimental/Workspace/Source/WorkspaceEditor/Private/SWorkspaceView.cpp
--- a/Engine/Plugins/Experimental/Workspace/Source/WorkspaceEditor/Private/SWorkspaceView.cpp
+++ b/Engine/Plugins/Experimental/Workspace/Source/WorkspaceEditor/Private/SWorkspaceView.cpp
@@ -94,17 +94,15 @@ public:
 				AddAssetButton->SetIsMenuOpen(false, false);
 			});
 
-			TArray<FAssetData> WorkspaceAssetDataEntries;
-			Workspace->GetAssetDataEntries(WorkspaceAssetDataEntries);
-	
-			AssetPickerConfig.OnShouldFilterAsset = FOnShouldFilterAsset::CreateLambda([this, WorkspaceAssetDataEntries, Workspace](const FAssetData& InAssetData)
+			TArray<FAssetData> AssetDataEntries(Workspace->GetAssetDataEntriesView());
+			AssetPickerConfig.OnShouldFilterAsset = FOnShouldFilterAsset::CreateLambda([AssetDataEntries, Workspace=WeakWorkspace](const FAssetData& InAssetData)
 			{
-				if (WorkspaceAssetDataEntries.Contains(InAssetData))
+				if (AssetDataEntries.Contains(InAssetData))
 				{
 					return true;
 				}
 	
-				return !Workspace->IsAssetSupported(InAssetData);
+				return Workspace.Get() ? !Workspace.Get()->IsAssetSupported(InAssetData) : false;
 			});
 
 			return ContentBrowserModule.Get().CreateAssetPicker(AssetPickerConfig);
@@ -253,4 +251,4 @@ void SWorkspaceView::GetHierarchyAssetData(TArray<FAssetData>& OutAssetData) con
 }
 }
 
-#undef LOCTEXT_NAMESPACE
\ No newline at end of file
+#undef LOCTEXT_NAMESPACE
diff --git a/Engine/Plugins/Experimental/Workspace/Source/WorkspaceEditor/Private/Workspace.cpp b/Engine/Plugins/Experimental/Workspace/Source/WorkspaceEditor/Private/Workspace.cpp
--- a/Engine/Plugins/Experimental/Workspace/Source/WorkspaceEditor/Private/Workspace.cpp
+++ b/Engine/Plugins/Experimental/Workspace/Source/WorkspaceEditor/Private/Workspace.cpp
@@ -268,6 +268,8 @@ UWorkspaceState* UWorkspace::GetState() const
 
 void UWorkspace::BroadcastModified()
 {
+	EntriesAssetData.Reset();
+	
 	if(!bSuspendNotifications)
 	{
 		ModifiedDelegate.Broadcast(this);
@@ -284,16 +286,26 @@ void UWorkspace::GetAssets(TArray<TObjectPtr<UObject>>& OutAssets) const
 			{
 				OutAssets.Add(Asset);
 			}
-		}	
+		}
 	}
 }
 
-void UWorkspace::GetAssetDataEntries(TArray<FAssetData>& OutAssetDataEntries) const
+TConstArrayView<FAssetData> UWorkspace::GetAssetDataEntriesView(bool bRefresh) const
 {
-	FARFilter Filter;
-	Algo::Transform(AssetEntries, Filter.SoftObjectPaths, [](const UWorkspaceAssetEntry* AssetEntry) { return AssetEntry != nullptr ? AssetEntry->Asset.ToSoftObjectPath() : nullptr; });
-	const IAssetRegistry& AssetRegistry = FAssetRegistryModule::GetRegistry();
-	AssetRegistry.GetAssets(Filter, OutAssetDataEntries);
+	if (EntriesAssetData.Num() == 0 || bRefresh)
+	{
+		EntriesAssetData.Reset();
+		const IAssetRegistry& AssetRegistry = FAssetRegistryModule::GetRegistry();
+		for (const UWorkspaceAssetEntry* AssetEntry : AssetEntries)
+		{
+			if (AssetEntry != nullptr)
+			{
+				EntriesAssetData.Add(AssetRegistry.GetAssetByObjectPath(AssetEntry->Asset.ToSoftObjectPath()));
+			}
+		}
+	}
+	
+	return EntriesAssetData;
 }
 
 bool UWorkspace::HasValidEntries() const
@@ -329,6 +341,8 @@ void UWorkspace::PostLoadExternalPackages()
 		check(IsValid(InLoadedEntry));
 		AssetEntries.Add(InLoadedEntry);
 	});
+
+	EntriesAssetData.Reset();
 }
 
 void UWorkspace::Serialize(FArchive& Ar)
@@ -409,4 +423,4 @@ UObject* UWorkspace::GetRootAsset() const
 	check(RootAsset);
 	
 	return RootAsset;
-}
\ No newline at end of file
+}
diff --git a/Engine/Plugins/Experimental/Workspace/Source/WorkspaceEditor/Private/Workspace.h b/Engine/Plugins/Experimental/Workspace/Source/WorkspaceEditor/Private/Workspace.h
--- a/Engine/Plugins/Experimental/Workspace/Source/WorkspaceEditor/Private/Workspace.h
+++ b/Engine/Plugins/Experimental/Workspace/Source/WorkspaceEditor/Private/Workspace.h
@@ -131,7 +131,7 @@ class UWorkspace : public UObject
 	// Returns all contained Asset objects
 	void GetAssets(TArray<TObjectPtr<UObject>>& OutAssets) const;
 	// Returns all contained AssetData entries
-	void GetAssetDataEntries(TArray<FAssetData>& OutAssetDataEntries) const;
+	TConstArrayView<FAssetData> GetAssetDataEntriesView(bool bRefresh = false) const;
 
 	// Returns whether or not the workspace contains _any_ valid UWorkspaceAssetEntry
 	bool HasValidEntries() const;
@@ -154,6 +154,8 @@ class UWorkspace : public UObject
 	// All of the assets in-directly referenced by this workspace
 	UPROPERTY(transient)
 	TArray<TObjectPtr<UWorkspaceAssetEntry>> AssetEntries;
+	
+	mutable TArray<FAssetData> EntriesAssetData;
 
 	// Schema for this workspace
 	UPROPERTY(AssetRegistrySearchable)
