UAF: Fix graphs and systems not being delete-able after having been opened

#Synced-CL 49556441

diff --git a/Engine/Plugins/Experimental/Workspace/Source/WorkspaceEditor/Private/SWorkspaceTabWrapper.cpp b/Engine/Plugins/Experimental/Workspace/Source/WorkspaceEditor/Private/SWorkspaceTabWrapper.cpp
--- a/Engine/Plugins/Experimental/Workspace/Source/WorkspaceEditor/Private/SWorkspaceTabWrapper.cpp
+++ b/Engine/Plugins/Experimental/Workspace/Source/WorkspaceEditor/Private/SWorkspaceTabWrapper.cpp
@@ -46,8 +46,7 @@ void SWorkspaceTabWrapper::Construct( const FArguments& InArgs, TSharedPtr<class
 		{
 			if (UWorkspaceTabMenuContext* Context = InMenu->FindContext<UWorkspaceTabMenuContext>())
 			{
-				TSharedRef<SWorkspaceTabWrapper> This = Context->TabWrapper.Pin().ToSharedRef();
-				
+				const TSharedRef<SWorkspaceTabWrapper> This = Context->TabWrapper.Pin().ToSharedRef();
 				const FToolMenuEntry SaveEntry = FToolMenuEntry::InitToolBarButton(
 					"Save",
 					FUIAction(
@@ -61,13 +60,7 @@ void SWorkspaceTabWrapper::Construct( const FArguments& InArgs, TSharedPtr<class
 
 				const FToolMenuEntry FindInContentBrowserEntry = FToolMenuEntry::InitToolBarButton(
 					"FindInContentBrowser",
-					FUIAction(FExecuteAction::CreateSPLambda(This, [This]() 
-						{
-							if (UObject* Asset = This->WeakDocumentObject.Get())
-							{
-								GEditor->SyncBrowserToObject(Asset);
-							}
-						})),
+					FUIAction(FExecuteAction::CreateSP(This, &SWorkspaceTabWrapper::FindObjectInContentBrowser)),
 					FText::GetEmpty(),
 					LOCTEXT("FindInContentBrowserTooltip", "Finds this asset in the content browser"),
 					FSlateIcon(FAppStyle::Get().GetStyleSetName(), "SystemWideCommands.FindInContentBrowser"));
@@ -375,4 +368,12 @@ bool SWorkspaceTabWrapper::IsSaveButtonVisible() const
 	return false;
 }
 
+void SWorkspaceTabWrapper::FindObjectInContentBrowser() const
+{
+	if (WeakDocumentObject.IsValid() && GEditor)
+	{
+		GEditor->SyncBrowserToObject(WeakDocumentObject.Get());
+	}
+}
+
 #undef LOCTEXT_NAMESPACE // "SWorkspaceTabWrapper"
diff --git a/Engine/Plugins/Experimental/Workspace/Source/WorkspaceEditor/Private/SWorkspaceTabWrapper.h b/Engine/Plugins/Experimental/Workspace/Source/WorkspaceEditor/Private/SWorkspaceTabWrapper.h
--- a/Engine/Plugins/Experimental/Workspace/Source/WorkspaceEditor/Private/SWorkspaceTabWrapper.h
+++ b/Engine/Plugins/Experimental/Workspace/Source/WorkspaceEditor/Private/SWorkspaceTabWrapper.h
@@ -43,6 +43,7 @@ protected:
 	void ExecuteSave() const;
 	bool CanExecuteSave() const;
 	bool IsSaveButtonVisible() const;
+	void FindObjectInContentBrowser() const;
 
 	TSharedPtr<SWidget> Content = nullptr;
 	TWeakPtr<UE::Workspace::FWorkspaceEditor> WeakWorkspaceEditor = nullptr;
