Workspace Preview:
- Pass PreviewScene to OnExit to allow unregistering of previously registered components and keep the component list from growing on reinit
- Fix leaking preview scene

#Synced-CL 50858266

diff --git a/Engine/Plugins/Experimental/UAF/UAF/Source/UAFEditor/Private/UAF/Viewport/SystemViewportController.cpp b/Engine/Plugins/Experimental/UAF/UAF/Source/UAFEditor/Private/UAF/Viewport/SystemViewportController.cpp
--- a/Engine/Plugins/Experimental/UAF/UAF/Source/UAFEditor/Private/UAF/Viewport/SystemViewportController.cpp
+++ b/Engine/Plugins/Experimental/UAF/UAF/Source/UAFEditor/Private/UAF/Viewport/SystemViewportController.cpp
@@ -63,14 +63,23 @@ namespace UE::UAF::Editor
 		}
 	}
 	
-	void FSystemViewportController::OnExit()
+	void FSystemViewportController::OnExit(FAdvancedPreviewScene* PreviewScene)
 	{
 		for (AActor* PreviewActor : PreviewActors)
 		{
-			PreviewActor->Destroy();
-			PreviewActor = nullptr;
+			if (PreviewActor)
+			{
+				if (PreviewScene && PreviewActor->GetRootComponent())
+				{
+					PreviewScene->RemoveComponent(PreviewActor->GetRootComponent());
+				}
+			
+				PreviewActor->Destroy();
+				PreviewActor = nullptr;
+			}
 		}
 		
+		
 		PreviewActors.Empty();
 	}
 }
diff --git a/Engine/Plugins/Experimental/UAF/UAF/Source/UAFEditor/Private/UAF/Viewport/SystemViewportController.h b/Engine/Plugins/Experimental/UAF/UAF/Source/UAFEditor/Private/UAF/Viewport/SystemViewportController.h
--- a/Engine/Plugins/Experimental/UAF/UAF/Source/UAFEditor/Private/UAF/Viewport/SystemViewportController.h
+++ b/Engine/Plugins/Experimental/UAF/UAF/Source/UAFEditor/Private/UAF/Viewport/SystemViewportController.h
@@ -9,7 +9,7 @@ namespace UE::UAF::Editor
 	class FSystemViewportController : public Workspace::IWorkspaceViewportController
 	{
 		virtual void OnEnter(const FViewportContext& InViewportContext) override;
-		virtual void OnExit() override;
+		virtual void OnExit(FAdvancedPreviewScene* PreviewScene) override;
 
 		TArray<AActor*> PreviewActors;
 	};
diff --git a/Engine/Plugins/Experimental/UAF/UAFAnimGraph/Source/UAFAnimGraphEditor/Private/AnimGraphViewportController.cpp b/Engine/Plugins/Experimental/UAF/UAFAnimGraph/Source/UAFAnimGraphEditor/Private/AnimGraphViewportController.cpp
--- a/Engine/Plugins/Experimental/UAF/UAFAnimGraph/Source/UAFAnimGraphEditor/Private/AnimGraphViewportController.cpp
+++ b/Engine/Plugins/Experimental/UAF/UAFAnimGraph/Source/UAFAnimGraphEditor/Private/AnimGraphViewportController.cpp
@@ -67,14 +67,23 @@ namespace UE::UAF::Editor
 		}
 	}
 	
-	void FAnimGraphViewportController::OnExit()
+	void FAnimGraphViewportController::OnExit(FAdvancedPreviewScene* PreviewScene)
 	{
 		for (AActor* PreviewActor : PreviewActors)
 		{
-			PreviewActor->Destroy();
-			PreviewActor = nullptr;
+			if (PreviewActor)
+			{
+				if (PreviewScene && PreviewActor->GetRootComponent())
+				{
+					PreviewScene->RemoveComponent(PreviewActor->GetRootComponent());
+				}
+			
+				PreviewActor->Destroy();
+				PreviewActor = nullptr;
+			}
 		}
 		
+		
 		PreviewActors.Empty();
 	}
 }
diff --git a/Engine/Plugins/Experimental/UAF/UAFAnimGraph/Source/UAFAnimGraphEditor/Private/AnimGraphViewportController.h b/Engine/Plugins/Experimental/UAF/UAFAnimGraph/Source/UAFAnimGraphEditor/Private/AnimGraphViewportController.h
--- a/Engine/Plugins/Experimental/UAF/UAFAnimGraph/Source/UAFAnimGraphEditor/Private/AnimGraphViewportController.h
+++ b/Engine/Plugins/Experimental/UAF/UAFAnimGraph/Source/UAFAnimGraphEditor/Private/AnimGraphViewportController.h
@@ -9,7 +9,7 @@ namespace UE::UAF::Editor
 	class FAnimGraphViewportController : public Workspace::IWorkspaceViewportController
 	{
 		virtual void OnEnter(const FViewportContext& InViewportContext) override;
-		virtual void OnExit() override;
+		virtual void OnExit(FAdvancedPreviewScene* PreviewScene) override;
 
 		TArray<AActor*> PreviewActors;
 	};
diff --git a/Engine/Plugins/Experimental/Workspace/Source/WorkspaceEditor/Private/WorkspaceAssetViewportClient.h b/Engine/Plugins/Experimental/Workspace/Source/WorkspaceEditor/Private/WorkspaceAssetViewportClient.h
--- a/Engine/Plugins/Experimental/Workspace/Source/WorkspaceEditor/Private/WorkspaceAssetViewportClient.h
+++ b/Engine/Plugins/Experimental/Workspace/Source/WorkspaceEditor/Private/WorkspaceAssetViewportClient.h
@@ -2,6 +2,7 @@
 
 #pragma once
 
+#include "AdvancedPreviewScene.h"
 #include "EditorViewportClient.h"
 
 namespace UE::Workspace
@@ -9,12 +10,15 @@ namespace UE::Workspace
 	class FWorkspaceAssetViewportClient : public FEditorViewportClient
 	{
 	public:
-		FWorkspaceAssetViewportClient(FEditorModeTools* InModeTools, FPreviewScene* InPreviewScene)
-			: FEditorViewportClient(InModeTools, InPreviewScene)
+		FWorkspaceAssetViewportClient(FEditorModeTools* InModeTools, const TSharedPtr<FAdvancedPreviewScene>& InPreviewScene)
+			: FEditorViewportClient(InModeTools, InPreviewScene.Get())
+			, AdvancedPreviewScene(InPreviewScene)
 		{
 		}
 		
 		// FEditorViewportClient interface
 		virtual void Tick(float DeltaTime) override;
+		
+		TSharedPtr<FAdvancedPreviewScene> AdvancedPreviewScene = nullptr;
 	};
 }
\ No newline at end of file
diff --git a/Engine/Plugins/Experimental/Workspace/Source/WorkspaceEditor/Private/WorkspaceEditor.cpp b/Engine/Plugins/Experimental/Workspace/Source/WorkspaceEditor/Private/WorkspaceEditor.cpp
--- a/Engine/Plugins/Experimental/Workspace/Source/WorkspaceEditor/Private/WorkspaceEditor.cpp
+++ b/Engine/Plugins/Experimental/Workspace/Source/WorkspaceEditor/Private/WorkspaceEditor.cpp
@@ -663,9 +663,10 @@ const FWorkspaceDocument& FWorkspaceEditor::GetFocusedWorkspaceDocument() const
 	
 void FWorkspaceEditor::ReinitViewportScene()
 {
+	FAdvancedPreviewScene* PreviewScene = static_cast<FAdvancedPreviewScene*>(ViewportClient->GetPreviewScene());
 	if (ActiveViewportController)
 	{
-		ActiveViewportController->OnExit();
+		ActiveViewportController->OnExit(PreviewScene);
 		ActiveViewportController = nullptr;
 	}
 
@@ -683,7 +684,7 @@ void FWorkspaceEditor::ReinitViewportScene()
 			if (ActiveViewportController)
 			{
 				IWorkspaceViewportController::FViewportContext ViewportContext;
-				ViewportContext.PreviewScene = static_cast<FAdvancedPreviewScene*>(ViewportClient->GetPreviewScene());
+				ViewportContext.PreviewScene = PreviewScene;
 				ViewportContext.OutlinerObject = AssetObject;
 				ViewportContext.SceneDescription = SceneDescription;
 				
@@ -1272,11 +1273,11 @@ AssetEditorViewportFactoryFunction FWorkspaceEditor::GetViewportDelegate()
 	
 TSharedPtr<FEditorViewportClient> FWorkspaceEditor::CreateEditorViewportClient() const
 {
-	FAdvancedPreviewScene* NewPreviewScene = new FAdvancedPreviewScene(FPreviewScene::ConstructionValues());
+	TSharedPtr<FAdvancedPreviewScene> NewAdvancedPreviewScene = MakeShareable(new FAdvancedPreviewScene(FPreviewScene::ConstructionValues()));
 	
-	StaticCastSharedPtr<FAssetEditorModeManager>(EditorModeManager)->SetPreviewScene(NewPreviewScene);
+	StaticCastSharedPtr<FAssetEditorModeManager>(EditorModeManager)->SetPreviewScene(NewAdvancedPreviewScene.Get());
 
-	TSharedPtr<FEditorViewportClient> EditorViewportClient = MakeShared<FWorkspaceAssetViewportClient>(EditorModeManager.Get(), NewPreviewScene);
+	TSharedPtr<FEditorViewportClient> EditorViewportClient = MakeShared<FWorkspaceAssetViewportClient>(EditorModeManager.Get(), NewAdvancedPreviewScene);
 	EditorViewportClient->ViewportType = LVT_Perspective;
 	EditorViewportClient->bSetListenerPosition = false;
 	EditorViewportClient->SetViewLocation(FVector(0.0f, 400.0f, 200.0f));
diff --git a/Engine/Plugins/Experimental/Workspace/Source/WorkspaceEditor/Public/WorkspaceViewportController.h b/Engine/Plugins/Experimental/Workspace/Source/WorkspaceEditor/Public/WorkspaceViewportController.h
--- a/Engine/Plugins/Experimental/Workspace/Source/WorkspaceEditor/Public/WorkspaceViewportController.h
+++ b/Engine/Plugins/Experimental/Workspace/Source/WorkspaceEditor/Public/WorkspaceViewportController.h
@@ -33,7 +33,7 @@ namespace UE::Workspace
 		};
 		
 		virtual void OnEnter(const FViewportContext& InContext) = 0;
-		virtual void OnExit() = 0;
+		virtual void OnExit(FAdvancedPreviewScene* PreviewScene) = 0;
 	};
 
 	using FWorkspaceViewportControllerFactory = TFunctionRef<TUniquePtr<IWorkspaceViewportController>()>;
