Workspace: Checking for IsAsset and IsSubobject are mutually exclusive and could prevent a valid outliner export to be processed correctly

#Synced-CL 51290764

diff --git a/Engine/Plugins/Experimental/Workspace/Source/WorkspaceEditor/Public/WorkspaceAssetRegistryInfo.h b/Engine/Plugins/Experimental/Workspace/Source/WorkspaceEditor/Public/WorkspaceAssetRegistryInfo.h
--- a/Engine/Plugins/Experimental/Workspace/Source/WorkspaceEditor/Public/WorkspaceAssetRegistryInfo.h
+++ b/Engine/Plugins/Experimental/Workspace/Source/WorkspaceEditor/Public/WorkspaceAssetRegistryInfo.h
@@ -223,7 +223,7 @@ public:
 			for (int32 SegmentIndex = Path.PathSegments.Num() - 1; SegmentIndex >= 0; SegmentIndex--)
 			{
 				FSoftObjectPath ObjectPath = Path.PathSegments[SegmentIndex].ToString();
-				if (ObjectPath.IsValid() && ObjectPath.IsAsset())
+				if (ObjectPath.IsValid())
 				{
 					if (ObjectPath.IsSubobject())
 					{
