EditorUsability : Workspace
- Fixed tab location not saved due to wrong object passed to the Asset Editor Toolkit, transient workspace was passed instead of underlying asset object resulting in tab location not to be retrieved properly in FAssetEditorToolkit::InitAssetEditor

#Synced-CL 51539177

diff --git a/Engine/Plugins/Experimental/Workspace/Source/WorkspaceEditor/Private/Workspace.h b/Engine/Plugins/Experimental/Workspace/Source/WorkspaceEditor/Private/Workspace.h
--- a/Engine/Plugins/Experimental/Workspace/Source/WorkspaceEditor/Private/Workspace.h
+++ b/Engine/Plugins/Experimental/Workspace/Source/WorkspaceEditor/Private/Workspace.h
@@ -80,6 +80,7 @@ class UWorkspace : public UObject
 	friend class UE::Workspace::SWorkspaceOutliner;
 	
 	friend class UWorkspaceFactory;
+	friend class UWorkspaceAssetEditor;
 	friend class UAssetDefinition_Workspace;
 	friend class UWorkspaceState;
 
diff --git a/Engine/Plugins/Experimental/Workspace/Source/WorkspaceEditor/Private/WorkspaceAssetEditor.cpp b/Engine/Plugins/Experimental/Workspace/Source/WorkspaceEditor/Private/WorkspaceAssetEditor.cpp
--- a/Engine/Plugins/Experimental/Workspace/Source/WorkspaceEditor/Private/WorkspaceAssetEditor.cpp
+++ b/Engine/Plugins/Experimental/Workspace/Source/WorkspaceEditor/Private/WorkspaceAssetEditor.cpp
@@ -2,12 +2,29 @@
 
 #include "WorkspaceAssetEditor.h"
 #include "WorkspaceEditor.h"
+#include "WorkspaceSchema.h"
 
 #include UE_INLINE_GENERATED_CPP_BY_NAME(WorkspaceAssetEditor)
 
 void UWorkspaceAssetEditor::GetObjectsToEdit(TArray<UObject*>& OutObjectsToEdit)
 {
-	OutObjectsToEdit.Add(ObjectToEdit);
+	if (ObjectToEdit)
+	{
+		if (const UWorkspaceSchema* Schema = ObjectToEdit->GetSchema())
+		{
+			if (Schema->CanSaveWorkspace())
+			{
+				OutObjectsToEdit.Add(ObjectToEdit);
+			}
+			else
+			{
+				// Transient workspace outputs its assets
+				TArray<TObjectPtr<UObject>> AssetObjects;
+				ObjectToEdit->GetAssets(AssetObjects);
+				OutObjectsToEdit.Append(AssetObjects);
+			}
+		}
+	}
 }
 
 void UWorkspaceAssetEditor::SetObjectToEdit(UWorkspace* InWorkspace)
diff --git a/Engine/Plugins/Experimental/Workspace/Source/WorkspaceEditor/Private/WorkspaceAssetEditor.h b/Engine/Plugins/Experimental/Workspace/Source/WorkspaceEditor/Private/WorkspaceAssetEditor.h
--- a/Engine/Plugins/Experimental/Workspace/Source/WorkspaceEditor/Private/WorkspaceAssetEditor.h
+++ b/Engine/Plugins/Experimental/Workspace/Source/WorkspaceEditor/Private/WorkspaceAssetEditor.h
@@ -15,7 +15,9 @@ class UWorkspaceAssetEditor : public UAssetEditor
 public:
 	void SetObjectToEdit(UWorkspace* InWorkspace);
 	UWorkspace* GetObjectToEdit();
+	
 protected:
+	/** For standalone workspace, the workspace itself will be returned otherwise the underlying assets will be returned */
 	virtual void GetObjectsToEdit(TArray<UObject*>& OutObjectsToEdit) override;
 	virtual TSharedPtr<FBaseAssetToolkit> CreateToolkit() override;
 
diff --git a/Engine/Plugins/Experimental/Workspace/Source/WorkspaceEditor/Private/WorkspaceEditorModule.cpp b/Engine/Plugins/Experimental/Workspace/Source/WorkspaceEditor/Private/WorkspaceEditorModule.cpp
--- a/Engine/Plugins/Experimental/Workspace/Source/WorkspaceEditor/Private/WorkspaceEditorModule.cpp
+++ b/Engine/Plugins/Experimental/Workspace/Source/WorkspaceEditor/Private/WorkspaceEditorModule.cpp
@@ -330,8 +330,11 @@ bool FWorkspaceEditorModule::GetExportedAssetsForWorkspace(const FAssetData& InW
 		AssetEditor->SetObjectToEdit(NewWorkspace);
 		AssetEditor->Initialize();
 		
-		IWorkspaceEditor* WorkspaceEditor = static_cast<IWorkspaceEditor*>(GEditor->GetEditorSubsystem<UAssetEditorSubsystem>()->FindEditorForAsset(NewWorkspace, true));
-		WorkspaceEditor->OpenAssets({InObject});
+		IWorkspaceEditor* WorkspaceEditor = static_cast<IWorkspaceEditor*>(GEditor->GetEditorSubsystem<UAssetEditorSubsystem>()->FindEditorForAsset(InObject, true));
+		if (WorkspaceEditor)
+		{
+			WorkspaceEditor->OpenAssets({InObject});
+		}
 		
 		return WorkspaceEditor;
 	}
