Request a details panel refresh instead of forcing it as it can be called after property changes are being propagated. If we refresh too early, we can invalidate property handles that are needed by callbacks.

#Synced-CL 52204442

diff --git a/Engine/Plugins/Experimental/Workspace/Source/WorkspaceEditor/Private/WorkspaceEditor.cpp b/Engine/Plugins/Experimental/Workspace/Source/WorkspaceEditor/Private/WorkspaceEditor.cpp
--- a/Engine/Plugins/Experimental/Workspace/Source/WorkspaceEditor/Private/WorkspaceEditor.cpp
+++ b/Engine/Plugins/Experimental/Workspace/Source/WorkspaceEditor/Private/WorkspaceEditor.cpp
@@ -520,7 +520,7 @@ void FWorkspaceEditor::RefreshDetails()
 {
 	if (DetailsView.IsValid())
 	{
-		DetailsView->ForceRefresh();
+		DetailsView->RequestForceRefresh();
 	}
 }
 
